Fixed: some more android rendering issues on ARPG in Vulkan
Disabled paralllel shadows on android

diff --git a/Engine/Config/Android/AndroidEngine.ini b/Engine/Config/Android/AndroidEngine.ini
--- a/Engine/Config/Android/AndroidEngine.ini
+++ b/Engine/Config/Android/AndroidEngine.ini
@@ -71,5 +71,6 @@ DefaultPlatformService=AndroidLocalNotification
 [ConsoleVariables]
 Slate.CacheRenderData=0
 r.VSync=1
+r.ParallelShadows=0
 
 ;GalaxyS6Edge=577
