[NO NON-PII COMMIT MESSAGE]

diff --git a/Engine/Config/Android/AndroidEngine.ini b/Engine/Config/Android/AndroidEngine.ini
--- a/Engine/Config/Android/AndroidEngine.ini
+++ b/Engine/Config/Android/AndroidEngine.ini
@@ -16,6 +16,7 @@ AudioDeviceModuleName=AndroidAudio
 AudioMixerModuleName=AudioMixerAndroid
 ; Defines a platform-specific volume headroom (in dB) for audio to provide better platform consistency with respect to volume levels.
 PlatformHeadroomDB=-6
+UseAudioMixer=true
 
 [PlatformMemoryBuckets]
 LargestMemoryBucket_MinGB=8
