- More CIS/jira error fixes (couple more coming)

diff --git a/Engine/Config/Android/DataDrivenPlatformInfo.ini b/Engine/Config/Android/DataDrivenPlatformInfo.ini
--- a/Engine/Config/Android/DataDrivenPlatformInfo.ini
+++ b/Engine/Config/Android/DataDrivenPlatformInfo.ini
@@ -45,7 +45,7 @@ DisplayName=Android (ATC)
 PlatformType=Game
 PlatformFlags=CookFlavor
 NormalIconPath=Launcher/Android/Platform_Android_ATC_24x
-LargeIconPath=Launcher/Android/Platform_Android_ATC_128x
+LargeIconPath=Launcher/Android/Platform_Android_128x
 XLargeIconPath=
 UATCommandLine=-targetplatform=Android -cookflavor=ATC
 AutoSDKPath=Android
@@ -64,7 +64,7 @@ DisplayName=Android (DXT)
 PlatformType=Game
 PlatformFlags=CookFlavor
 NormalIconPath=Launcher/Android/Platform_Android_DXT_24x
-LargeIconPath=Launcher/Android/Platform_Android_DXT_128x
+LargeIconPath=Launcher/Android/Platform_Android_128x
 XLargeIconPath=
 UATCommandLine=-targetplatform=Android -cookflavor=DXT
 AutoSDKPath=Android
@@ -83,7 +83,7 @@ DisplayName=Android (ETC1)
 PlatformType=Game
 PlatformFlags=CookFlavor
 NormalIconPath=Launcher/Android/Platform_Android_ETC1_24x
-LargeIconPath=Launcher/Android/Platform_Android_ETC1_128x
+LargeIconPath=Launcher/Android/Platform_Android_128x
 XLargeIconPath=
 UATCommandLine=-targetplatform=Android -cookflavor=ETC1
 AutoSDKPath=Android
@@ -101,8 +101,8 @@ TargetPlatformName=Android_ETC1a
 DisplayName=Android (ETC1a)
 PlatformType=Game
 PlatformFlags=CookFlavor
-NormalIconPath=Launcher/Android/Platform_Android_ETC1a_24x
-LargeIconPath=Launcher/Android/Platform_Android_ETC1a_128x
+NormalIconPath=Launcher/Android/Platform_Android_ETC1_24x
+LargeIconPath=Launcher/Android/Platform_Android_128x
 XLargeIconPath=
 UATCommandLine=-targetplatform=Android -cookflavor=ETC1a
 AutoSDKPath=Android
@@ -121,7 +121,7 @@ DisplayName=Android (ETC2)
 PlatformType=Game
 PlatformFlags=CookFlavor
 NormalIconPath=Launcher/Android/Platform_Android_ETC2_24x
-LargeIconPath=Launcher/Android/Platform_Android_ETC2_128x
+LargeIconPath=Launcher/Android/Platform_Android_128x
 XLargeIconPath=
 UATCommandLine=-targetplatform=Android -cookflavor=ETC2
 AutoSDKPath=Android
@@ -140,7 +140,7 @@ DisplayName=Android (PVRTC)
 PlatformType=Game
 PlatformFlags=CookFlavor
 NormalIconPath=Launcher/Android/Platform_Android_PVRTC_24x
-LargeIconPath=Launcher/Android/Platform_Android_PVRTC_128x
+LargeIconPath=Launcher/Android/Platform_Android_128x
 XLargeIconPath=
 UATCommandLine=-targetplatform=Android -cookflavor=PVRTC
 AutoSDKPath=Android
@@ -159,7 +159,7 @@ DisplayName=Android (ASTC)
 PlatformType=Game
 PlatformFlags=CookFlavor
 NormalIconPath=Launcher/Android/Platform_Android_ASTC_24x
-LargeIconPath=Launcher/Android/Platform_Android_ASTC_128x
+LargeIconPath=Launcher/Android/Platform_Android_128x
 XLargeIconPath=
 UATCommandLine=-targetplatform=Android -cookflavor=ASTC
 AutoSDKPath=Android
@@ -217,7 +217,7 @@ DisplayName=Android (ATC) (Client-only)
 PlatformType=Game
 PlatformFlags=CookFlavor
 NormalIconPath=Launcher/Android/Platform_Android_ATC_24x
-LargeIconPath=Launcher/Android/Platform_Android_ATC_128x
+LargeIconPath=Launcher/Android/Platform_Android_128x
 XLargeIconPath=
 UATCommandLine=-client -targetplatform=Android -cookflavor=ATC
 AutoSDKPath=Android
@@ -236,7 +236,7 @@ DisplayName=Android (DXT) (Client-only)
 PlatformType=Game
 PlatformFlags=CookFlavor
 NormalIconPath=Launcher/Android/Platform_Android_DXT_24x
-LargeIconPath=Launcher/Android/Platform_Android_DXT_128x
+LargeIconPath=Launcher/Android/Platform_Android_128x
 XLargeIconPath=
 UATCommandLine=-client -targetplatform=Android -cookflavor=DXT
 AutoSDKPath=Android
@@ -255,7 +255,7 @@ DisplayName=Android (ETC1) (Client-only)
 PlatformType=Game
 PlatformFlags=CookFlavor
 NormalIconPath=Launcher/Android/Platform_Android_ETC1_24x
-LargeIconPath=Launcher/Android/Platform_Android_ETC1_128x
+LargeIconPath=Launcher/Android/Platform_Android_128x
 XLargeIconPath=
 UATCommandLine=-client -targetplatform=Android -cookflavor=ETC1
 AutoSDKPath=Android
@@ -274,7 +274,7 @@ DisplayName=Android (ETC1a) (Client-only)
 PlatformType=Game
 PlatformFlags=CookFlavor
 NormalIconPath=Launcher/Android/Platform_Android_ETC1a_24x
-LargeIconPath=Launcher/Android/Platform_Android_ETC1a_128x
+LargeIconPath=Launcher/Android/Platform_Android_128x
 XLargeIconPath=
 UATCommandLine=-client -targetplatform=Android -cookflavor=ETC1a
 AutoSDKPath=Android
@@ -293,7 +293,7 @@ DisplayName=Android (ETC2) (Client-only)
 PlatformType=Game
 PlatformFlags=CookFlavor
 NormalIconPath=Launcher/Android/Platform_Android_ETC2_24x
-LargeIconPath=Launcher/Android/Platform_Android_ETC2_128x
+LargeIconPath=Launcher/Android/Platform_Android_128x
 XLargeIconPath=
 UATCommandLine=-client -targetplatform=Android -cookflavor=ETC2
 AutoSDKPath=Android
@@ -312,7 +312,7 @@ DisplayName=Android (PVRTC) (Client-only)
 PlatformType=Game
 PlatformFlags=CookFlavor
 NormalIconPath=Launcher/Android/Platform_Android_PVRTC_24x
-LargeIconPath=Launcher/Android/Platform_Android_PVRTC_128x
+LargeIconPath=Launcher/Android/Platform_Android_128x
 XLargeIconPath=
 UATCommandLine=-client -targetplatform=Android -cookflavor=PVRTC
 AutoSDKPath=Android
@@ -331,7 +331,7 @@ DisplayName=Android (ASTC) (Client-only)
 PlatformType=Game
 PlatformFlags=CookFlavor
 NormalIconPath=Launcher/Android/Platform_Android_ASTC_24x
-LargeIconPath=Launcher/Android/Platform_Android_ASTC_128x
+LargeIconPath=Launcher/Android/Platform_Android_128x
 XLargeIconPath=
 UATCommandLine=-client -targetplatform=Android -cookflavor=ASTC
 AutoSDKPath=Android
diff --git a/Engine/Plugins/Lumin/MagicLeap/Source/MagicLeap/Private/AsyncDestroyer.cpp b/Engine/Plugins/Lumin/MagicLeap/Source/MagicLeap/Private/AsyncDestroyer.cpp
--- a/Engine/Plugins/Lumin/MagicLeap/Source/MagicLeap/Private/AsyncDestroyer.cpp
+++ b/Engine/Plugins/Lumin/MagicLeap/Source/MagicLeap/Private/AsyncDestroyer.cpp
@@ -4,7 +4,7 @@
 #include "Engine/Engine.h"
 #include "AppEventHandler.h"
 #if PLATFORM_LUMIN
-#include "Lumin/LuminAffinity.h"
+#include "Lumin/LuminPlatformAffinity.h"
 #endif // PLATFORM_LUMIN
 
 namespace MagicLeap
diff --git a/Engine/Plugins/Lumin/MagicLeap/Source/MagicLeap/Private/CameraCaptureComponent.cpp b/Engine/Plugins/Lumin/MagicLeap/Source/MagicLeap/Private/CameraCaptureComponent.cpp
--- a/Engine/Plugins/Lumin/MagicLeap/Source/MagicLeap/Private/CameraCaptureComponent.cpp
+++ b/Engine/Plugins/Lumin/MagicLeap/Source/MagicLeap/Private/CameraCaptureComponent.cpp
@@ -11,7 +11,7 @@
 #include "Containers/Array.h"
 #include "Lumin/LuminPlatformFile.h"
 #if PLATFORM_LUMIN
-#include "Lumin/LuminAffinity.h"
+#include "Lumin/LuminPlatformAffinity.h"
 #endif // PLATFORM_LUMIN
 #include "MagicLeapPluginUtil.h" // for ML_INCLUDES_START/END
 
diff --git a/Engine/Plugins/Lumin/MagicLeap/Source/MagicLeap/Private/CameraCaptureRunnable.cpp b/Engine/Plugins/Lumin/MagicLeap/Source/MagicLeap/Private/CameraCaptureRunnable.cpp
--- a/Engine/Plugins/Lumin/MagicLeap/Source/MagicLeap/Private/CameraCaptureRunnable.cpp
+++ b/Engine/Plugins/Lumin/MagicLeap/Source/MagicLeap/Private/CameraCaptureRunnable.cpp
@@ -40,7 +40,7 @@
 #include "AppFramework.h"
 #include "Lumin/LuminPlatformFile.h"
 #if PLATFORM_LUMIN
-#include "Lumin/LuminAffinity.h"
+#include "Lumin/LuminPlatformAffinity.h"
 #endif // PLATFORM_LUMIN
 #include "MagicLeapPluginUtil.h" // for ML_INCLUDES_START/END
 
diff --git a/Engine/Plugins/Lumin/MagicLeap/Source/MagicLeap/Private/ImageTrackerComponent.cpp b/Engine/Plugins/Lumin/MagicLeap/Source/MagicLeap/Private/ImageTrackerComponent.cpp
--- a/Engine/Plugins/Lumin/MagicLeap/Source/MagicLeap/Private/ImageTrackerComponent.cpp
+++ b/Engine/Plugins/Lumin/MagicLeap/Source/MagicLeap/Private/ImageTrackerComponent.cpp
@@ -10,7 +10,7 @@
 #include "HAL/Runnable.h"
 #include "HAL/RunnableThread.h"
 #if PLATFORM_LUMIN
-#include "Lumin/LuminAffinity.h"
+#include "Lumin/LuminPlatformAffinity.h"
 #endif // PLATFORM_LUMIN
 #if WITH_EDITOR
 #include "Editor.h"
diff --git a/Engine/Plugins/Lumin/MagicLeap/Source/MagicLeap/Private/ImageTrackerRunnable.cpp b/Engine/Plugins/Lumin/MagicLeap/Source/MagicLeap/Private/ImageTrackerRunnable.cpp
--- a/Engine/Plugins/Lumin/MagicLeap/Source/MagicLeap/Private/ImageTrackerRunnable.cpp
+++ b/Engine/Plugins/Lumin/MagicLeap/Source/MagicLeap/Private/ImageTrackerRunnable.cpp
@@ -40,7 +40,7 @@
 #include "Misc/Paths.h"
 #include "Engine/Texture2D.h"
 #if PLATFORM_LUMIN
-#include "Lumin/LuminAffinity.h"
+#include "Lumin/LuminPlatformAffinity.h"
 #endif // PLATFORM_LUMIN
 
 #if WITH_MLSDK
diff --git a/Engine/Plugins/Lumin/MagicLeap/Source/MagicLeap/Private/InAppPurchaseComponent.cpp b/Engine/Plugins/Lumin/MagicLeap/Source/MagicLeap/Private/InAppPurchaseComponent.cpp
--- a/Engine/Plugins/Lumin/MagicLeap/Source/MagicLeap/Private/InAppPurchaseComponent.cpp
+++ b/Engine/Plugins/Lumin/MagicLeap/Source/MagicLeap/Private/InAppPurchaseComponent.cpp
@@ -6,7 +6,7 @@
 #include "Modules/ModuleManager.h"
 #include "AppEventHandler.h"
 #if PLATFORM_LUMIN
-#include "Lumin/LuminAffinity.h"
+#include "Lumin/LuminPlatformAffinity.h"
 #endif // PLATFORM_LUMIN
 
 #if WITH_MLSDK
diff --git a/Engine/Plugins/Lumin/MagicLeap/Source/MagicLeap/Private/MagicLeapARPinComponent.cpp b/Engine/Plugins/Lumin/MagicLeap/Source/MagicLeap/Private/MagicLeapARPinComponent.cpp
--- a/Engine/Plugins/Lumin/MagicLeap/Source/MagicLeap/Private/MagicLeapARPinComponent.cpp
+++ b/Engine/Plugins/Lumin/MagicLeap/Source/MagicLeap/Private/MagicLeapARPinComponent.cpp
@@ -46,7 +46,7 @@
 #include "HeadMountedDisplayFunctionLibrary.h"
 #include "MagicLeapSecureStorage.h"
 #if PLATFORM_LUMIN
-#include "Lumin/LuminAffinity.h"
+#include "Lumin/LuminPlatformAffinity.h"
 #endif // PLATFORM_LUMIN
 
 #if WITH_MLSDK
diff --git a/Engine/Plugins/Lumin/MagicLeapMedia/Source/MagicLeapMediaCodec/Private/MediaCodecInputWorker.cpp b/Engine/Plugins/Lumin/MagicLeapMedia/Source/MagicLeapMediaCodec/Private/MediaCodecInputWorker.cpp
--- a/Engine/Plugins/Lumin/MagicLeapMedia/Source/MagicLeapMediaCodec/Private/MediaCodecInputWorker.cpp
+++ b/Engine/Plugins/Lumin/MagicLeapMedia/Source/MagicLeapMediaCodec/Private/MediaCodecInputWorker.cpp
@@ -5,7 +5,7 @@
 #include "HAL/RunnableThread.h"
 #include "HAL/Event.h"
 #include "MagicLeapMediaCodecPlayer.h"
-#include "Lumin/LuminAffinity.h"
+#include "Lumin/LuminPlatformAffinity.h"
 #include "IMagicLeapMediaCodecModule.h"
 
 #include "ml_media_extractor.h"
diff --git a/Engine/Plugins/Lumin/MagicLeapScreens/Source/MagicLeapScreens/Private/MagicLeapScreensComponent.cpp b/Engine/Plugins/Lumin/MagicLeapScreens/Source/MagicLeapScreens/Private/MagicLeapScreensComponent.cpp
--- a/Engine/Plugins/Lumin/MagicLeapScreens/Source/MagicLeapScreens/Private/MagicLeapScreensComponent.cpp
+++ b/Engine/Plugins/Lumin/MagicLeapScreens/Source/MagicLeapScreens/Private/MagicLeapScreensComponent.cpp
@@ -14,7 +14,7 @@
 #include "MagicLeapScreensWorker.h"
 #include "MagicLeapScreensMsg.h"
 #if PLATFORM_LUMIN
-#include "Lumin/LuminAffinity.h"
+#include "Lumin/LuminPlatformAffinity.h"
 #endif // PLATFORM_LUMIN
 #include "MagicLeapPluginUtil.h" // for ML_INCLUDES_START/END
 
diff --git a/Engine/Source/Developer/DesktopPlatform/Private/PlatformInfo.cpp b/Engine/Source/Developer/DesktopPlatform/Private/PlatformInfo.cpp
--- a/Engine/Source/Developer/DesktopPlatform/Private/PlatformInfo.cpp
+++ b/Engine/Source/Developer/DesktopPlatform/Private/PlatformInfo.cpp
@@ -18,6 +18,10 @@ namespace
 
 TArray<FPlatformInfo> AllPlatformInfoArray;
 
+// we don't need any of this without the editor, although we would ideally not even compile this outside of the editor
+// @todo platplug: Figure out why this is compiled on target devices
+#if WITH_EDITOR
+
 void BuildPlatformInfo(const FName& InPlatformInfoName, const FName& InTargetPlatformName, const FText& InDisplayName, const EPlatformType InPlatformType, const EPlatformFlags::Flags InPlatformFlags, const FPlatformIconPaths& InIconPaths, const FString& InUATCommandLine, const FString& InAutoSDKPath, EPlatformSDKStatus InStatus, const FString& InTutorial, bool InEnabled, FString InBinaryFolderName, FString InIniPlatformName, bool InUsesHostCompiler, bool InUATClosesAfterLaunch, bool InIsConfidential, const FName& InUBTTargetId, const FName& InPlatformGroupName)
 {
 	FPlatformInfo& PlatformInfo = AllPlatformInfoArray.Emplace_GetRef();
@@ -64,7 +68,10 @@ void BuildPlatformInfo(const FName& InPlatformInfoName, const FName& InTargetPla
 	// Generate the icon style names for FEditorStyle
 	PlatformInfo.IconPaths.NormalStyleName = *FString::Printf(TEXT("Launcher.Platform_%s"), *InPlatformInfoNameString);
 	PlatformInfo.IconPaths.LargeStyleName  = *FString::Printf(TEXT("Launcher.Platform_%s.Large"), *InPlatformInfoNameString);
-	PlatformInfo.IconPaths.XLargeStyleName = *FString::Printf(TEXT("Launcher.Platform_%s.XLarge"), *InPlatformInfoNameString);
+	if (PlatformInfo.IconPaths.XLargePath != TEXT(""))
+	{
+		PlatformInfo.IconPaths.XLargeStyleName = *FString::Printf(TEXT("Launcher.Platform_%s.XLarge"), *InPlatformInfoNameString);
+	}
 
 	// SDK data
 	PlatformInfo.SDKStatus = InStatus;
@@ -112,11 +119,6 @@ EPlatformFlags::Flags ConvertPlatformFlags(const FString& String)
 
 void ParseDataDrivenPlatformInfo(const TCHAR* Name, const FConfigSection& Section)
 {
-	checkf(WITH_EDITOR, TEXT("We don't expect to oarse DDPI on non-editor"));
-//	void BuildPlatformInfo(const FName& InPlatformInfoName, const FName& InTargetPlatformName, const FText& InDisplayName, const EPlatformType InPlatformType, const EPlatformFlags::Flags InPlatformFlags, 
-//		const FPlatformIconPaths& InIconPaths, const FString& InUATCommandLine, const FString& InAutoSDKPath, EPlatformSDKStatus InStatus, const FString& InTutorial, bool InEnabled, FString InBinaryFolderName, 
-//		FString InIniPlatformName, bool InUsesHostCompiler, bool InUATClosesAfterLaunch, bool InIsConfidential, const FName& InUBTTargetId, const FName& InPlatformGroupName)
-
 	// @todo platplug: use FNames instead of TCHAR* for keys, so we don't ahve to re-convert every time
 	FName TargetPlatformName = *GetSectionString(Section, TEXT("TargetPlatformName"));
 	FString DisplayName = GetSectionString(Section, TEXT("DisplayName"));
@@ -176,6 +178,7 @@ struct FPlatformInfoAutoInit
 	}
 
 } GPlatformInfoAutoInit;
+#endif
 
 } // anonymous namespace
 
diff --git a/Engine/Source/Runtime/Core/Private/Lumin/LuminOutputDevices.cpp b/Engine/Source/Runtime/Core/Private/Lumin/LuminOutputDevices.cpp
deleted file mode 100644
diff --git a/Engine/Source/Runtime/Core/Private/Lumin/LuminPlatformOutputDevices.cpp b/Engine/Source/Runtime/Core/Private/Lumin/LuminPlatformOutputDevices.cpp
new file mode 100644
--- /dev/null
+++ b/Engine/Source/Runtime/Core/Private/Lumin/LuminPlatformOutputDevices.cpp
@@ -0,0 +1,44 @@
+// Copyright 1998-2019 Epic Games, Inc. All Rights Reserved.
+// Copyright 2017 Magic Leap, Inc. All Rights Reserved.
+
+#include "Lumin/LuminPlatformOutputDevices.h"
+#include "Lumin/LuminOutputDeviceDebug.h"
+#include "CoreGlobals.h"
+#include "Misc/Parse.h"
+#include "Templates/ScopedPointer.h"
+#include "Misc/CommandLine.h"
+#include "Misc/Paths.h"
+#include "Misc/OutputDeviceMemory.h"
+#include "Misc/OutputDeviceFile.h"
+#include "Misc/OutputDeviceDebug.h"
+#include "Misc/OutputDeviceAnsiError.h"
+#include "Misc/App.h"
+#include "HAL/FeedbackContextAnsi.h"
+#include "Misc/OutputDeviceConsole.h"
+
+void FLuminOutputDevices::SetupOutputDevices()
+{
+	check(GLog);
+
+	CachedAbsoluteFilename[0] = 0;
+
+	GLog->AddOutputDevice(FPlatformOutputDevices::GetLog());
+
+#if !NO_LOGGING
+	// if console is enabled add an output device, unless the commandline says otherwise...
+	if (GLogConsole && !FParse::Param(FCommandLine::Get(), TEXT("NOCONSOLE")))
+	{
+		GLog->AddOutputDevice(GLogConsole);
+	}
+	
+	// If the platform has a separate debug output channel (e.g. OutputDebugString) then add an output device
+	// unless logging is turned off
+	if (FPlatformMisc::HasSeparateChannelForDebugOutput())
+	{
+        // Use FLuminOutputDeviceDebug instead of the default FOutputDeviceDebug so that ml_log can respect the verbosity in the output
+		GLog->AddOutputDevice(new FLuminOutputDeviceDebug());
+	}
+#endif
+
+	GLog->AddOutputDevice(FPlatformOutputDevices::GetEventLog());
+}
diff --git a/Engine/Source/Runtime/Core/Public/Android/AndroidPlatformOutputDevices.h b/Engine/Source/Runtime/Core/Public/Android/AndroidPlatformOutputDevices.h
--- a/Engine/Source/Runtime/Core/Public/Android/AndroidPlatformOutputDevices.h
+++ b/Engine/Source/Runtime/Core/Public/Android/AndroidPlatformOutputDevices.h
@@ -12,6 +12,4 @@ struct CORE_API FAndroidOutputDevices : public FGenericPlatformOutputDevices
 {
 };
 
-#if !PLATFORM_LUMIN
 typedef FAndroidOutputDevices FPlatformOutputDevices;
-#endif
diff --git a/Engine/Source/Runtime/Core/Public/Linux/LinuxPlatformAffinity.h b/Engine/Source/Runtime/Core/Public/Linux/LinuxPlatformAffinity.h
--- a/Engine/Source/Runtime/Core/Public/Linux/LinuxPlatformAffinity.h
+++ b/Engine/Source/Runtime/Core/Public/Linux/LinuxPlatformAffinity.h
@@ -1,65 +1,6 @@
 // Copyright 1998-2019 Epic Games, Inc. All Rights Reserved.
 
-/*=============================================================================================
-LuminAffinity.h: Lumin affinity profile masks definitions.
-==============================================================================================*/
-
 #pragma once
 
 #include "GenericPlatform/GenericPlatformAffinity.h"
-
-class FLuminAffinity : public FGenericPlatformAffinity
-{
-	const static uint64 ArmCores = MAKEAFFINITYMASK2(3,4);
-	const static uint64 DenverCores = MAKEAFFINITYMASK1(2);
-public:
-	static const CORE_API uint64 GetMainGameMask() {
-		return ArmCores;
-	}
-
-	static const CORE_API uint64 GetRenderingThreadMask() {
-		return ArmCores;
-	}
-
-	static const CORE_API uint64 GetRHIThreadMask() {
-		return ArmCores;
-	}
-
-	static const CORE_API uint64 GetRTHeartBeatMask() {
-		return ArmCores;
-	}
-
-	static const CORE_API uint64 GetPoolThreadMask() {
-		return DenverCores;
-	}
-
-	static const CORE_API uint64 GetTaskGraphThreadMask() {
-		return ArmCores;
-	}
-
-	static const CORE_API uint64 GetStatsThreadMask() {
-		return ArmCores;
-	}
-
-	static const CORE_API uint64 GetAudioThreadMask() {
-		return DenverCores;
-	}
-
-	static const CORE_API uint64 GetNoAffinityMask() {
-		return 0xFFFFFFFFFFFFFFFF;
-	}
-
-	static const CORE_API uint64 GetTaskGraphBackgroundTaskMask() {
-		return ArmCores;
-	}
-
-	static EThreadPriority GetRenderingThreadPriority() {
-		return TPri_Normal;
-	}
-
-	static EThreadPriority GetRHIThreadPriority() {
-		return TPri_SlightlyBelowNormal;
-	}
-};
-
-typedef FLuminAffinity FPlatformAffinity;
+typedef FGenericPlatformAffinity FPlatformAffinity;
diff --git a/Engine/Source/Runtime/Core/Public/Lumin/LuminOutputDevices.h b/Engine/Source/Runtime/Core/Public/Lumin/LuminOutputDevices.h
deleted file mode 100644
diff --git a/Engine/Source/Runtime/Core/Public/Lumin/LuminPlatformAffinity.h b/Engine/Source/Runtime/Core/Public/Lumin/LuminPlatformAffinity.h
--- a/Engine/Source/Runtime/Core/Public/Lumin/LuminPlatformAffinity.h
+++ b/Engine/Source/Runtime/Core/Public/Lumin/LuminPlatformAffinity.h
@@ -1,9 +1,5 @@
 // Copyright 1998-2019 Epic Games, Inc. All Rights Reserved.
 
-/*=============================================================================================
-LuminAffinity.h: Lumin affinity profile masks definitions.
-==============================================================================================*/
-
 #pragma once
 
 #include "GenericPlatform/GenericPlatformAffinity.h"
diff --git a/Engine/Source/Runtime/Core/Public/Lumin/LuminPlatformOutputDevices.h b/Engine/Source/Runtime/Core/Public/Lumin/LuminPlatformOutputDevices.h
new file mode 100644
--- /dev/null
+++ b/Engine/Source/Runtime/Core/Public/Lumin/LuminPlatformOutputDevices.h
@@ -0,0 +1,12 @@
+// Copyright 1998-2019 Epic Games, Inc. All Rights Reserved.
+// Copyright 2017 Magic Leap, Inc. All Rights Reserved.
+
+#pragma once
+#include "GenericPlatform/GenericPlatformOutputDevices.h"
+
+struct CORE_API FLuminOutputDevices : public FGenericPlatformOutputDevices
+{
+	static void	SetupOutputDevices();
+};
+
+typedef FLuminOutputDevices FPlatformOutputDevices;
diff --git a/Engine/Source/Runtime/RenderCore/Private/RenderUtils.cpp b/Engine/Source/Runtime/RenderCore/Private/RenderUtils.cpp
--- a/Engine/Source/Runtime/RenderCore/Private/RenderUtils.cpp
+++ b/Engine/Source/Runtime/RenderCore/Private/RenderUtils.cpp
@@ -928,19 +928,19 @@ static TAutoConsoleVariable<int32> CVarDistanceFields(
 	); 
 
 
-RENDERCORE_API uint32 GForwardShadingPlatformMask = 0;
+RENDERCORE_API uint64 GForwardShadingPlatformMask = 0;
 static_assert(SP_NumPlatforms <= sizeof(GForwardShadingPlatformMask) * 8, "GForwardShadingPlatformMask must be large enough to support all shader platforms");
 
-RENDERCORE_API uint32 GDBufferPlatformMask = 0;
+RENDERCORE_API uint64 GDBufferPlatformMask = 0;
 static_assert(SP_NumPlatforms <= sizeof(GDBufferPlatformMask) * 8, "GDBufferPlatformMask must be large enough to support all shader platforms");
 
-RENDERCORE_API uint32 GBasePassVelocityPlatformMask = 0;
+RENDERCORE_API uint64 GBasePassVelocityPlatformMask = 0;
 static_assert(SP_NumPlatforms <= sizeof(GBasePassVelocityPlatformMask) * 8, "GBasePassVelocityPlatformMask must be large enough to support all shader platforms");
 
-RENDERCORE_API uint32 GSelectiveBasePassOutputsPlatformMask = 0;
+RENDERCORE_API uint64 GSelectiveBasePassOutputsPlatformMask = 0;
 static_assert(SP_NumPlatforms <= sizeof(GSelectiveBasePassOutputsPlatformMask) * 8, "GSelectiveBasePassOutputsPlatformMask must be large enough to support all shader platforms");
 
-RENDERCORE_API uint32 GDistanceFieldsPlatformMask = 0;
+RENDERCORE_API uint64 GDistanceFieldsPlatformMask = 0;
 static_assert(SP_NumPlatforms <= sizeof(GDistanceFieldsPlatformMask) * 8, "GDistanceFieldsPlatformMask must be large enough to support all shader platforms");
 
 RENDERCORE_API void RenderUtilsInit()
@@ -985,7 +985,7 @@ RENDERCORE_API void RenderUtilsInit()
 			ITargetPlatform* TargetPlatform = TargetPlatformManager->FindTargetPlatform(PlatformName.ToString());
 			if (TargetPlatform)
 			{
-				uint32 Mask = 1u << ShaderPlatformIndex;
+				uint64 Mask = 1ull << ShaderPlatformIndex;
 
 				if (TargetPlatform->UsesForwardShading())
 				{
diff --git a/Engine/Source/Runtime/RenderCore/Public/RenderUtils.h b/Engine/Source/Runtime/RenderCore/Public/RenderUtils.h
--- a/Engine/Source/Runtime/RenderCore/Public/RenderUtils.h
+++ b/Engine/Source/Runtime/RenderCore/Public/RenderUtils.h
@@ -433,8 +433,8 @@ RENDERCORE_API bool AllowPixelDepthOffset(EShaderPlatform Platform);
 /** Returns if ForwardShading is enabled. Only valid for the current platform (otherwise call ITargetPlatform::UsesForwardShading()). */
 inline bool IsForwardShadingEnabled(EShaderPlatform Platform)
 {
-	extern RENDERCORE_API uint32 GForwardShadingPlatformMask;
-	return !!(GForwardShadingPlatformMask & (1u << Platform))
+	extern RENDERCORE_API uint64 GForwardShadingPlatformMask;
+	return !!(GForwardShadingPlatformMask & (1ull << Platform))
 		// Culling uses compute shader
 		&& GetMaxSupportedFeatureLevel(Platform) >= ERHIFeatureLevel::SM5;
 }
@@ -454,29 +454,29 @@ inline bool IsUsingGBuffers(EShaderPlatform Platform)
 /** Returns whether DBuffer decals are enabled for a given shader platform */
 inline bool IsUsingDBuffers(EShaderPlatform Platform)
 {
-	extern RENDERCORE_API uint32 GDBufferPlatformMask;
-	return !!(GDBufferPlatformMask & (1u << Platform));
+	extern RENDERCORE_API uint64 GDBufferPlatformMask;
+	return !!(GDBufferPlatformMask & (1ull << Platform));
 }
 
 /** Returns whether the base pass should output to the velocity buffer is enabled for a given shader platform */
 inline bool IsUsingBasePassVelocity(EShaderPlatform Platform)
 {
-	extern RENDERCORE_API uint32 GBasePassVelocityPlatformMask;
-	return !!(GBasePassVelocityPlatformMask & (1u << Platform));
+	extern RENDERCORE_API uint64 GBasePassVelocityPlatformMask;
+	return !!(GBasePassVelocityPlatformMask & (1ull << Platform));
 }
 
 /** Returns whether the base pass should use selective outputs for a given shader platform */
 inline bool IsUsingSelectiveBasePassOutputs(EShaderPlatform Platform)
 {
-	extern RENDERCORE_API uint32 GSelectiveBasePassOutputsPlatformMask;
-	return !!(GSelectiveBasePassOutputsPlatformMask & (1u << Platform));
+	extern RENDERCORE_API uint64 GSelectiveBasePassOutputsPlatformMask;
+	return !!(GSelectiveBasePassOutputsPlatformMask & (1ull << Platform));
 }
 
 /** Returns whether distance fields are enabled for a given shader platform */
 inline bool IsUsingDistanceFields(EShaderPlatform Platform)
 {
-	extern RENDERCORE_API uint32 GDistanceFieldsPlatformMask;
-	return !!(GDistanceFieldsPlatformMask & (1u << Platform));
+	extern RENDERCORE_API uint64 GDistanceFieldsPlatformMask;
+	return !!(GDistanceFieldsPlatformMask & (1ull << Platform));
 }
 
 inline bool IsUsingPerPixelDBufferMask(EShaderPlatform Platform)
