- Fixed the AutomationToolLauncher.csproj error
- Missed a file in last checkin

diff --git a/Engine/Build/BatchFiles/FindPlatformExtensionSources.bat b/Engine/Build/BatchFiles/FindPlatformExtensionSources.bat
--- a/Engine/Build/BatchFiles/FindPlatformExtensionSources.bat
+++ b/Engine/Build/BatchFiles/FindPlatformExtensionSources.bat
@@ -1,4 +1,6 @@
 @echo off
+setlocal
+
 rem ## Unreal Engine 4 AutomationTool setup script
 rem ## Copyright 1998-2019 Epic Games, Inc. All Rights Reserved.
 rem ##
diff --git a/Engine/Config/Android/DataDrivenPlatformInfo.ini b/Engine/Config/Android/DataDrivenPlatformInfo.ini
--- a/Engine/Config/Android/DataDrivenPlatformInfo.ini
+++ b/Engine/Config/Android/DataDrivenPlatformInfo.ini
@@ -176,7 +176,7 @@ PlatformGroupName=Mobile
 [PlatformInfo AndroidClient]
 TargetPlatformName=AndroidClient
 DisplayName=Android (Client-only)
-PlatformType=Game
+PlatformType=Client
 PlatformFlags=None
 NormalIconPath=Launcher/Android/Platform_Android_24x
 LargeIconPath=Launcher/Android/Platform_Android_128x
@@ -195,7 +195,7 @@ PlatformGroupName=Mobile
 [PlatformInfo Android_MultiClient]
 TargetPlatformName=Android_MultiClient
 DisplayName=Android (Multi) (Client-only)
-PlatformType=Game
+PlatformType=Client
 PlatformFlags=CookFlavor
 NormalIconPath=Launcher/Android/Platform_Android_24x
 LargeIconPath=Launcher/Android/Platform_Android_128x
@@ -214,7 +214,7 @@ PlatformGroupName=Mobile
 [PlatformInfo Android_ATCClient]
 TargetPlatformName=Android_ATCClient
 DisplayName=Android (ATC) (Client-only)
-PlatformType=Game
+PlatformType=Client
 PlatformFlags=CookFlavor
 NormalIconPath=Launcher/Android/Platform_Android_ATC_24x
 LargeIconPath=Launcher/Android/Platform_Android_128x
@@ -233,7 +233,7 @@ PlatformGroupName=Mobile
 [PlatformInfo Android_DXTClient]
 TargetPlatformName=Android_DXTClient
 DisplayName=Android (DXT) (Client-only)
-PlatformType=Game
+PlatformType=Client
 PlatformFlags=CookFlavor
 NormalIconPath=Launcher/Android/Platform_Android_DXT_24x
 LargeIconPath=Launcher/Android/Platform_Android_128x
@@ -252,7 +252,7 @@ PlatformGroupName=Mobile
 [PlatformInfo Android_ETC1Client]
 TargetPlatformName=Android_ETC1Client
 DisplayName=Android (ETC1) (Client-only)
-PlatformType=Game
+PlatformType=Client
 PlatformFlags=CookFlavor
 NormalIconPath=Launcher/Android/Platform_Android_ETC1_24x
 LargeIconPath=Launcher/Android/Platform_Android_128x
@@ -271,7 +271,7 @@ PlatformGroupName=Mobile
 [PlatformInfo Android_ETC1aClient]
 TargetPlatformName=Android_ETC1aClient
 DisplayName=Android (ETC1a) (Client-only)
-PlatformType=Game
+PlatformType=Client
 PlatformFlags=CookFlavor
 NormalIconPath=Launcher/Android/Platform_Android_ETC1a_24x
 LargeIconPath=Launcher/Android/Platform_Android_128x
@@ -290,7 +290,7 @@ PlatformGroupName=Mobile
 [PlatformInfo Android_ETC2Client]
 TargetPlatformName=Android_ETC2Client
 DisplayName=Android (ETC2) (Client-only)
-PlatformType=Game
+PlatformType=Client
 PlatformFlags=CookFlavor
 NormalIconPath=Launcher/Android/Platform_Android_ETC2_24x
 LargeIconPath=Launcher/Android/Platform_Android_128x
@@ -309,7 +309,7 @@ PlatformGroupName=Mobile
 [PlatformInfo Android_PVRTCClient]
 TargetPlatformName=Android_PVRTCClient
 DisplayName=Android (PVRTC) (Client-only)
-PlatformType=Game
+PlatformType=Client
 PlatformFlags=CookFlavor
 NormalIconPath=Launcher/Android/Platform_Android_PVRTC_24x
 LargeIconPath=Launcher/Android/Platform_Android_128x
@@ -328,7 +328,7 @@ PlatformGroupName=Mobile
 [PlatformInfo Android_ASTCClient]
 TargetPlatformName=Android_ASTCClient
 DisplayName=Android (ASTC) (Client-only)
-PlatformType=Game
+PlatformType=Client
 PlatformFlags=CookFlavor
 NormalIconPath=Launcher/Android/Platform_Android_ASTC_24x
 LargeIconPath=Launcher/Android/Platform_Android_128x
