- Fixed bad Android icon reference

diff --git a/Engine/Config/Android/DataDrivenPlatformInfo.ini b/Engine/Config/Android/DataDrivenPlatformInfo.ini
--- a/Engine/Config/Android/DataDrivenPlatformInfo.ini
+++ b/Engine/Config/Android/DataDrivenPlatformInfo.ini
@@ -273,7 +273,7 @@ TargetPlatformName=Android_ETC1aClient
 DisplayName=Android (ETC1a) (Client-only)
 PlatformType=Client
 PlatformFlags=CookFlavor
-NormalIconPath=Launcher/Android/Platform_Android_ETC1a_24x
+NormalIconPath=Launcher/Android/Platform_Android_ETC1_24x
 LargeIconPath=Launcher/Android/Platform_Android_128x
 XLargeIconPath=
 UATCommandLine=-client -targetplatform=Android -cookflavor=ETC1a
