Cleaning up android engine ini file.

diff --git a/Engine/Config/Android/AndroidEngine.ini b/Engine/Config/Android/AndroidEngine.ini
--- a/Engine/Config/Android/AndroidEngine.ini
+++ b/Engine/Config/Android/AndroidEngine.ini
@@ -11,10 +11,10 @@ gc.MaxObjectsInGame=131072
 
 [Audio]
 AudioDeviceModuleName=AndroidAudio
-; Uncomment below (and comment above) to use multi-platform mixer module by default
-;AudioDeviceModuleName=AudioMixerAndroid
+
 ; Defining below allows switching to audio mixer using -audiomixer commandline
 AudioMixerModuleName=AudioMixerAndroid
+
 ; Defines a platform-specific volume headroom (in dB) for audio to provide better platform consistency with respect to volume levels.
 PlatformHeadroomDB=-6
 UseAudioMixer=true
