Added mobile specific, simplified SkyAtmosphere scalability settings

diff --git a/Engine/Config/Android/AndroidScalability.ini b/Engine/Config/Android/AndroidScalability.ini
new file mode 100644
--- /dev/null
+++ b/Engine/Config/Android/AndroidScalability.ini
@@ -0,0 +1,51 @@
+; ------------------------------------------------------------------------------
+; Effects settings
+; ------------------------------------------------------------------------------
+
+[EffectsQuality@0]
+r.SkyAtmosphere.AerialPerspectiveLUT.FastApplyOnOpaque=1 ; Always have FastSkyLUT 1 in this case to avoid wrong sky
+r.SkyAtmosphere.AerialPerspectiveLUT.SampleCountPerSlice=1
+r.SkyAtmosphere.AerialPerspectiveLUT.DepthResolution=4
+r.SkyAtmosphere.FastSkyLUT=1
+r.SkyAtmosphere.FastSkyLUT.SampleCountMin=1.0
+r.SkyAtmosphere.FastSkyLUT.SampleCountMax=8.0
+r.SkyAtmosphere.FastSkyLUT.Width=96
+r.SkyAtmosphere.FastSkyLUT.Height=50
+r.SkyAtmosphere.TransmittanceLUT.UseSmallFormat=1
+r.SkyAtmosphere.MultiScatteringLUT.UseSmallFormat=1
+
+[EffectsQuality@1]
+r.SkyAtmosphere.AerialPerspectiveLUT.FastApplyOnOpaque=1 ; Always have FastSkyLUT 1 in this case to avoid wrong sky
+r.SkyAtmosphere.AerialPerspectiveLUT.SampleCountPerSlice=1
+r.SkyAtmosphere.AerialPerspectiveLUT.DepthResolution=4
+r.SkyAtmosphere.FastSkyLUT=1
+r.SkyAtmosphere.FastSkyLUT.SampleCountMin=1.0
+r.SkyAtmosphere.FastSkyLUT.SampleCountMax=8.0
+r.SkyAtmosphere.FastSkyLUT.Width=96
+r.SkyAtmosphere.FastSkyLUT.Height=50
+r.SkyAtmosphere.TransmittanceLUT.UseSmallFormat=1
+r.SkyAtmosphere.MultiScatteringLUT.UseSmallFormat=1
+
+[EffectsQuality@2]
+r.SkyAtmosphere.AerialPerspectiveLUT.FastApplyOnOpaque=1 ; Always have FastSkyLUT 1 in this case to avoid wrong sky
+r.SkyAtmosphere.AerialPerspectiveLUT.SampleCountPerSlice=1
+r.SkyAtmosphere.AerialPerspectiveLUT.DepthResolution=4
+r.SkyAtmosphere.FastSkyLUT=1
+r.SkyAtmosphere.FastSkyLUT.SampleCountMin=1.0
+r.SkyAtmosphere.FastSkyLUT.SampleCountMax=8.0
+r.SkyAtmosphere.FastSkyLUT.Width=96
+r.SkyAtmosphere.FastSkyLUT.Height=50
+r.SkyAtmosphere.TransmittanceLUT.UseSmallFormat=1
+r.SkyAtmosphere.MultiScatteringLUT.UseSmallFormat=1
+
+[EffectsQuality@3]
+r.SkyAtmosphere.AerialPerspectiveLUT.FastApplyOnOpaque=1 ; Always have FastSkyLUT 1 in this case to avoid wrong sky
+r.SkyAtmosphere.AerialPerspectiveLUT.SampleCountPerSlice=1
+r.SkyAtmosphere.AerialPerspectiveLUT.DepthResolution=4
+r.SkyAtmosphere.FastSkyLUT=1
+r.SkyAtmosphere.FastSkyLUT.SampleCountMin=1.0
+r.SkyAtmosphere.FastSkyLUT.SampleCountMax=8.0
+r.SkyAtmosphere.FastSkyLUT.Width=96
+r.SkyAtmosphere.FastSkyLUT.Height=50
+r.SkyAtmosphere.TransmittanceLUT.UseSmallFormat=1
+r.SkyAtmosphere.MultiScatteringLUT.UseSmallFormat=1
\ No newline at end of file
diff --git a/Engine/Config/IOS/IOSScalability.ini b/Engine/Config/IOS/IOSScalability.ini
new file mode 100644
--- /dev/null
+++ b/Engine/Config/IOS/IOSScalability.ini
@@ -0,0 +1,51 @@
+; ------------------------------------------------------------------------------
+; Effects settings
+; ------------------------------------------------------------------------------
+
+[EffectsQuality@0]
+r.SkyAtmosphere.AerialPerspectiveLUT.FastApplyOnOpaque=1 ; Always have FastSkyLUT 1 in this case to avoid wrong sky
+r.SkyAtmosphere.AerialPerspectiveLUT.SampleCountPerSlice=1
+r.SkyAtmosphere.AerialPerspectiveLUT.DepthResolution=4
+r.SkyAtmosphere.FastSkyLUT=1
+r.SkyAtmosphere.FastSkyLUT.SampleCountMin=1.0
+r.SkyAtmosphere.FastSkyLUT.SampleCountMax=8.0
+r.SkyAtmosphere.FastSkyLUT.Width=96
+r.SkyAtmosphere.FastSkyLUT.Height=50
+r.SkyAtmosphere.TransmittanceLUT.UseSmallFormat=1
+r.SkyAtmosphere.MultiScatteringLUT.UseSmallFormat=1
+
+[EffectsQuality@1]
+r.SkyAtmosphere.AerialPerspectiveLUT.FastApplyOnOpaque=1 ; Always have FastSkyLUT 1 in this case to avoid wrong sky
+r.SkyAtmosphere.AerialPerspectiveLUT.SampleCountPerSlice=1
+r.SkyAtmosphere.AerialPerspectiveLUT.DepthResolution=4
+r.SkyAtmosphere.FastSkyLUT=1
+r.SkyAtmosphere.FastSkyLUT.SampleCountMin=1.0
+r.SkyAtmosphere.FastSkyLUT.SampleCountMax=8.0
+r.SkyAtmosphere.FastSkyLUT.Width=96
+r.SkyAtmosphere.FastSkyLUT.Height=50
+r.SkyAtmosphere.TransmittanceLUT.UseSmallFormat=1
+r.SkyAtmosphere.MultiScatteringLUT.UseSmallFormat=1
+
+[EffectsQuality@2]
+r.SkyAtmosphere.AerialPerspectiveLUT.FastApplyOnOpaque=1 ; Always have FastSkyLUT 1 in this case to avoid wrong sky
+r.SkyAtmosphere.AerialPerspectiveLUT.SampleCountPerSlice=1
+r.SkyAtmosphere.AerialPerspectiveLUT.DepthResolution=4
+r.SkyAtmosphere.FastSkyLUT=1
+r.SkyAtmosphere.FastSkyLUT.SampleCountMin=1.0
+r.SkyAtmosphere.FastSkyLUT.SampleCountMax=8.0
+r.SkyAtmosphere.FastSkyLUT.Width=96
+r.SkyAtmosphere.FastSkyLUT.Height=50
+r.SkyAtmosphere.TransmittanceLUT.UseSmallFormat=1
+r.SkyAtmosphere.MultiScatteringLUT.UseSmallFormat=1
+
+[EffectsQuality@3]
+r.SkyAtmosphere.AerialPerspectiveLUT.FastApplyOnOpaque=1 ; Always have FastSkyLUT 1 in this case to avoid wrong sky
+r.SkyAtmosphere.AerialPerspectiveLUT.SampleCountPerSlice=1
+r.SkyAtmosphere.AerialPerspectiveLUT.DepthResolution=4
+r.SkyAtmosphere.FastSkyLUT=1
+r.SkyAtmosphere.FastSkyLUT.SampleCountMin=1.0
+r.SkyAtmosphere.FastSkyLUT.SampleCountMax=8.0
+r.SkyAtmosphere.FastSkyLUT.Width=96
+r.SkyAtmosphere.FastSkyLUT.Height=50
+r.SkyAtmosphere.TransmittanceLUT.UseSmallFormat=1
+r.SkyAtmosphere.MultiScatteringLUT.UseSmallFormat=1
\ No newline at end of file
