Set minimum occlusion distance to 2000 on mobile to reduce artifatcs related to occlusion latency

diff --git a/Engine/Config/Android/AndroidScalability.ini b/Engine/Config/Android/AndroidScalability.ini
--- a/Engine/Config/Android/AndroidScalability.ini
+++ b/Engine/Config/Android/AndroidScalability.ini
@@ -27,15 +27,19 @@ foliage.DensityScale=1.0
 
 [ViewDistanceQuality@0]
 r.ViewDistanceScale=0.70
+r.NeverOcclusionTestDistance=2000
 
 [ViewDistanceQuality@1]
 r.ViewDistanceScale=1.0
+r.NeverOcclusionTestDistance=2000
 
 [ViewDistanceQuality@2]
 r.ViewDistanceScale=1.0
+r.NeverOcclusionTestDistance=2000
 
 [ViewDistanceQuality@3]
 r.ViewDistanceScale=1.0
+r.NeverOcclusionTestDistance=2000
 
 ; ------------------------------------------------------------------------------
 ; Shadow settings
diff --git a/Engine/Config/IOS/IOSScalability.ini b/Engine/Config/IOS/IOSScalability.ini
--- a/Engine/Config/IOS/IOSScalability.ini
+++ b/Engine/Config/IOS/IOSScalability.ini
@@ -27,15 +27,19 @@ foliage.DensityScale=1.0
 
 [ViewDistanceQuality@0]
 r.ViewDistanceScale=0.70
+r.NeverOcclusionTestDistance=2000
 
 [ViewDistanceQuality@1]
 r.ViewDistanceScale=1.0
+r.NeverOcclusionTestDistance=2000
 
 [ViewDistanceQuality@2]
 r.ViewDistanceScale=1.0
+r.NeverOcclusionTestDistance=2000
 
 [ViewDistanceQuality@3]
 r.ViewDistanceScale=1.0
+r.NeverOcclusionTestDistance=2000
 
 ; ------------------------------------------------------------------------------
 ; Shadow settings
