Do not alternate texture compression for iOS and Android

diff --git a/Engine/Config/Android/AndroidEngine.ini b/Engine/Config/Android/AndroidEngine.ini
--- a/Engine/Config/Android/AndroidEngine.ini
+++ b/Engine/Config/Android/AndroidEngine.ini
@@ -110,4 +110,8 @@ bUseGooglePlayBillingApiV2=true
 +SensorLocations=/sys/htc/cpu_temp
 
 [/Script/Engine.RendererSettings]
-r.DistanceFields=0
\ No newline at end of file
+r.DistanceFields=0
+
+[AlternateTextureCompression]
+TextureCompressionFormat=""
+TextureFormatPrefix=""
\ No newline at end of file
diff --git a/Engine/Config/IOS/IOSEngine.ini b/Engine/Config/IOS/IOSEngine.ini
--- a/Engine/Config/IOS/IOSEngine.ini
+++ b/Engine/Config/IOS/IOSEngine.ini
@@ -80,4 +80,8 @@ RetryResumeDataLimit=3;
 net.DisableIPv6=0
 
 [/Script/Engine.RendererSettings]
-r.DistanceFields=0
\ No newline at end of file
+r.DistanceFields=0
+
+[AlternateTextureCompression]
+TextureCompressionFormat=""
+TextureFormatPrefix=""
\ No newline at end of file
