Android and IOS set default compressor to Mermaid

diff --git a/Engine/Config/Android/BaseAndroidGame.ini b/Engine/Config/Android/BaseAndroidGame.ini
new file mode 100644
--- /dev/null
+++ b/Engine/Config/Android/BaseAndroidGame.ini
@@ -0,0 +1,3 @@
+[/Script/UnrealEd.ProjectPackagingSettings]
+PakFileCompressionMethod=Mermaid
+
diff --git a/Engine/Config/IOS/BaseIOSGame.ini b/Engine/Config/IOS/BaseIOSGame.ini
new file mode 100644
--- /dev/null
+++ b/Engine/Config/IOS/BaseIOSGame.ini
@@ -0,0 +1,3 @@
+[/Script/UnrealEd.ProjectPackagingSettings]
+PakFileCompressionMethod=Mermaid
+
