Set FXAA quality to 0 on mobile by default

diff --git a/Engine/Config/Android/AndroidScalability.ini b/Engine/Config/Android/AndroidScalability.ini
--- a/Engine/Config/Android/AndroidScalability.ini
+++ b/Engine/Config/Android/AndroidScalability.ini
@@ -109,6 +109,18 @@ r.LightShaftQuality=1
 r.Mobile.AmbientOcclusionQuality=1
 r.Mobile.PixelProjectedReflectionQuality=1
 
+[AntiAliasingQuality@0]
+r.FXAA.Quality=0
+
+[AntiAliasingQuality@1]
+r.FXAA.Quality=0
+
+[AntiAliasingQuality@2]
+r.FXAA.Quality=0
+
+[AntiAliasingQuality@3]
+r.FXAA.Quality=0
+
 ; ------------------------------------------------------------------------------
 ; Effects settings
 ; ------------------------------------------------------------------------------
diff --git a/Engine/Config/IOS/IOSScalability.ini b/Engine/Config/IOS/IOSScalability.ini
--- a/Engine/Config/IOS/IOSScalability.ini
+++ b/Engine/Config/IOS/IOSScalability.ini
@@ -109,6 +109,18 @@ r.LightShaftQuality=1
 r.Mobile.AmbientOcclusionQuality=1
 r.Mobile.PixelProjectedReflectionQuality=1
 
+[AntiAliasingQuality@0]
+r.FXAA.Quality=0
+
+[AntiAliasingQuality@1]
+r.FXAA.Quality=0
+
+[AntiAliasingQuality@2]
+r.FXAA.Quality=0
+
+[AntiAliasingQuality@3]
+r.FXAA.Quality=0
+
 ; ------------------------------------------------------------------------------
 ; Effects settings
 ; ------------------------------------------------------------------------------
