Create DeviceProfiles for Android Preview(Android_Preview_OpenGL, Android_Preview_Vulkan)
Android_Preview_OpenGL sets FX.AllowGPUSorting to be consistent with device
FX.AllowGPUSorting has ECVF_Preview flag now to allow this

diff --git a/Engine/Config/Android/DataDrivenPlatformInfo.ini b/Engine/Config/Android/DataDrivenPlatformInfo.ini
--- a/Engine/Config/Android/DataDrivenPlatformInfo.ini
+++ b/Engine/Config/Android/DataDrivenPlatformInfo.ini
@@ -29,7 +29,7 @@ InactiveIconName=LevelEditor.PreviewMode.AndroidES31.Disabled
 MenuText=NSLOCTEXT("PreviewPlatform", "PreviewMenuText_ES31", "Android ES 3.1")
 MenuTooltip=NSLOCTEXT("PreviewPlatform", "PreviewMenuTooltip_ES31", "Mobile preview using Android ES3.1 quality settings.")
 IconText=NSLOCTEXT("PreviewPlatform", "PreviewIconText_ES31", "Android ES3.1")
-DeviceProfileName=Android
+DeviceProfileName=Android_Preview_OpenGL
 
 [PreviewPlatform AndroidVulkan]
 EnabledCVar=ini:Engine:Android:/Script/AndroidRuntimeSettings.AndroidRuntimeSettings:bSupportsVulkan
@@ -40,7 +40,7 @@ InactiveIconName=LevelEditor.PreviewMode.AndroidVulkan.Disabled
 MenuText=NSLOCTEXT("PreviewPlatform", "PreviewMenuText_AndroidVulkan", "Android Vulkan")
 MenuTooltip=NSLOCTEXT("PreviewPlatform", "PreviewMenuTooltip_AndroidVulkan", "Mobile preview using Android Vulkan quality settings.")
 IconText=NSLOCTEXT("PreviewPlatform", "PreviewIconText_AndroidVulkan", "Android Vulkan")
-DeviceProfileName=Android
+DeviceProfileName=Android_Preview_Vulkan
 
 [PreviewPlatform AndroidVulkanSM5]
 EnabledCVar=ini:Engine:Android:/Script/AndroidRuntimeSettings.AndroidRuntimeSettings:bSupportsVulkanSM5
diff --git a/Engine/Config/BaseDeviceProfiles.ini b/Engine/Config/BaseDeviceProfiles.ini
--- a/Engine/Config/BaseDeviceProfiles.ini
+++ b/Engine/Config/BaseDeviceProfiles.ini
@@ -70,6 +70,8 @@
 +DeviceProfileNameAndTypes=LinuxServer,Linux
 +DeviceProfileNameAndTypes=LinuxArm64Server,LinuxArm64
 +DeviceProfileNameAndTypes=Android,Android
++DeviceProfileNameAndTypes=Android_Preview_OpenGL,Android
++DeviceProfileNameAndTypes=Android_Preview_Vulkan,Android
 +DeviceProfileNameAndTypes=Android_Low,Android
 +DeviceProfileNameAndTypes=Android_Mid,Android
 +DeviceProfileNameAndTypes=Android_High,Android
@@ -653,6 +655,16 @@ BaseProfileName=
 ; If any PreviewAllowlistCVars is set, cvars are denied by default.
 PreviewAllowlistCVars=sg.
 
+[Android_Preview_OpenGL DeviceProfile]
+DeviceType=Android
+BaseProfileName=Android
+; FX.AllowGPUSorting is hardcoded to 0 in OpenGLDevice.cpp to avoid driver bugs
++CVars=FX.AllowGPUSorting=0
+
+[Android_Preview_Vulkan DeviceProfile]
+DeviceType=Android
+BaseProfileName=Android
+
 [Android_Low DeviceProfile]
 DeviceType=Android
 BaseProfileName=Android
diff --git a/Engine/Source/Runtime/Engine/Private/Particles/FXSystem.cpp b/Engine/Source/Runtime/Engine/Private/Particles/FXSystem.cpp
--- a/Engine/Source/Runtime/Engine/Private/Particles/FXSystem.cpp
+++ b/Engine/Source/Runtime/Engine/Private/Particles/FXSystem.cpp
@@ -146,7 +146,7 @@ namespace FXConsoleVariables
 		TEXT("FX.AllowGPUSorting"),
 		bAllowGPUSorting,
 		TEXT("Allow particles to be sorted on the GPU."),
-		ECVF_ReadOnly
+		ECVF_ReadOnly | ECVF_Preview
 		);
 	FAutoConsoleVariableRef CVarFreezeGPUSimulation(
 		TEXT("FX.FreezeGPUSimulation"),
