Android: Do not backtrace vulkan memory allocations by default

diff --git a/Engine/Config/Android/AndroidEngine.ini b/Engine/Config/Android/AndroidEngine.ini
--- a/Engine/Config/Android/AndroidEngine.ini
+++ b/Engine/Config/Android/AndroidEngine.ini
@@ -91,6 +91,8 @@ r.HairStrands.Strands=0
 r.HairStrands.Binding=0
 ; R11G11B10 UAV is not supported on Android
 r.TemporalAA.R11G11B10History=0
+; Do not backtrace vulkan memory allocations
+r.Vulkan.MemoryBacktrace=0
 
 [OnlineSubsystemGooglePlay.Store]
 bSupportsInAppPurchasing=true
