- Add Disable Preview Menu Option
- GetFriendlyShaderPlatformName removed
- RHIShaderPlatformDefinitions.inl deleted
- Added FriendlyNames to all ShaderPlatforms in DDSPI
- MenuText now comes from FriendlyNames in DDSPI

diff --git a/Engine/Config/Android/DataDrivenPlatformInfo.ini b/Engine/Config/Android/DataDrivenPlatformInfo.ini
--- a/Engine/Config/Android/DataDrivenPlatformInfo.ini
+++ b/Engine/Config/Android/DataDrivenPlatformInfo.ini
@@ -27,7 +27,6 @@ ShaderPlatform=OPENGL_ES3_1_ANDROID
 ShaderFormat=GLSL_ES3_1_ANDROID
 ActiveIconName=LevelEditor.PreviewMode.AndroidES31.Enabled
 InactiveIconName=LevelEditor.PreviewMode.AndroidES31.Disabled
-MenuText=NSLOCTEXT("PreviewPlatform", "PreviewMenuText_ES31", "Android ES 3.1")
 MenuTooltip=NSLOCTEXT("PreviewPlatform", "PreviewMenuTooltip_ES31", "Mobile preview using Android ES3.1 quality settings.")
 IconText=NSLOCTEXT("PreviewPlatform", "PreviewIconText_ES31", "Android ES3.1")
 DeviceProfileName=Android_Preview_OpenGL
@@ -39,7 +38,6 @@ ShaderPlatform=VULKAN_ES3_1_ANDROID
 ShaderFormat=SF_VULKAN_ES31_ANDROID
 ActiveIconName=LevelEditor.PreviewMode.AndroidVulkan.Enabled
 InactiveIconName=LevelEditor.PreviewMode.AndroidVulkan.Disabled
-MenuText=NSLOCTEXT("PreviewPlatform", "PreviewMenuText_AndroidVulkan", "Android Vulkan")
 MenuTooltip=NSLOCTEXT("PreviewPlatform", "PreviewMenuTooltip_AndroidVulkan", "Mobile preview using Android Vulkan quality settings.")
 IconText=NSLOCTEXT("PreviewPlatform", "PreviewIconText_AndroidVulkan", "Android Vulkan")
 DeviceProfileName=Android_Preview_Vulkan
@@ -51,7 +49,6 @@ ShaderPlatform=VULKAN_SM5_ANDROID
 ShaderFormat=SF_VULKAN_SM5_ANDROID
 ActiveIconName=LevelEditor.PreviewMode.AndroidVulkanSM5.Enabled
 InactiveIconName=LevelEditor.PreviewMode.AndroidVulkanSM5.Disabled
-MenuText=NSLOCTEXT("PreviewPlatform", "PreviewMenuText_AndroidVulkanSM5", "Android Vulkan SM5")
 MenuTooltip=NSLOCTEXT("PreviewPlatform", "PreviewMenuTooltip_AndroidVulkanSM5", "Mobile preview using Android Vulkan SM5 quality settings.")
 IconText=NSLOCTEXT("PreviewPlatform", "PreviewIconText_AndroidVulkanSM5", "Android VK SM5")
 DeviceProfileName=Android
@@ -66,6 +63,7 @@ bTargetsTiledGPU=true
 bSupportsMobileMultiView=true
 bNeedsOfflineCompiler=true
 bSupportsManualVertexFetch=false
+FriendlyName=LOCTEXT("FriendlyShaderPlatformName_AndroidES31", "Android ES 3_1")
 
 [ShaderPlatform VULKAN_ES3_1_ANDROID]
 Language=Vulkan
@@ -77,6 +75,7 @@ bSupportsMobileMultiView=true
 bNeedsOfflineCompiler=true
 bSupportsIndexBufferUAVs=true
 bSupportsManualVertexFetch=false
+FriendlyName=LOCTEXT("FriendlyShaderPlatformName_AndroidVulkan", "Android Vulkan")
 
 [ShaderPlatform VULKAN_SM5_ANDROID]
 Language=Vulkan
@@ -88,3 +87,4 @@ bSupportsIndexBufferUAVs=true
 bSupportsManualVertexFetch=false
 bRequiresDisableForwardLocalLights=true
 bWaterUsesSimpleForwardShading=true
+FriendlyName=LOCTEXT("FriendlyShaderPlatformName_AndroidVulkanSM5", "Android Vulkan SM5")
diff --git a/Engine/Config/HoloLens/DataDrivenPlatformInfo.ini b/Engine/Config/HoloLens/DataDrivenPlatformInfo.ini
--- a/Engine/Config/HoloLens/DataDrivenPlatformInfo.ini
+++ b/Engine/Config/HoloLens/DataDrivenPlatformInfo.ini
@@ -32,3 +32,4 @@ ShaderFormat=D3D_ES3_1_HOLOLENS
 bIsMobile=true
 bSupportsInstancedStereo=true
 bSupportsManualVertexFetch=false
+FriendlyName=LOCTEXT("FriendlyShaderPlatformName_HololensES31", "Hololens Shader Model 3_1")
\ No newline at end of file
diff --git a/Engine/Config/IOS/DataDrivenPlatformInfo.ini b/Engine/Config/IOS/DataDrivenPlatformInfo.ini
--- a/Engine/Config/IOS/DataDrivenPlatformInfo.ini
+++ b/Engine/Config/IOS/DataDrivenPlatformInfo.ini
@@ -30,7 +30,6 @@ ShaderFormat=SF_METAL
 ShaderPlatform=METAL
 ActiveIconName=LevelEditor.PreviewMode.iOS.Enabled
 InactiveIconName=LevelEditor.PreviewMode.iOS.Disabled
-MenuText=NSLOCTEXT("PreviewPlatform", "PreviewMenuText_iOS", "iOS")
 MenuTooltip=NSLOCTEXT("PreviewPlatform", "PreviewMenuTooltip_iOS", "Mobile preview using iOS material quality settings.")
 IconText=NSLOCTEXT("PreviewPlatform", "PreviewIconText_iOS", "iOS")
 DeviceProfileName=IOS
@@ -42,7 +41,6 @@ ShaderPlatform=METAL_MRT
 EnabledCVar=ini:Engine:IOS:/Script/IOSRuntimeSettings.IOSRuntimeSettings:bSupportsMetalMRT
 ActiveIconName=LevelEditor.PreviewMode.iOSSM5.Enabled
 InactiveIconName=LevelEditor.PreviewMode.iOSSM5.Disabled
-MenuText=NSLOCTEXT("PreviewPlatform", "PreviewMenuText_iOSSM5", "iOS (SM5 Renderer)")
 MenuTooltip=NSLOCTEXT("PreviewPlatform", "PreviewMenuTooltip_iOSSM5", "Mobile preview using iOS SM5 material quality settings.")
 IconText=NSLOCTEXT("PreviewPlatform", "PreviewIconText_iOSSM5", "iOS SM5")
 DeviceProfileName=IOS
@@ -56,6 +54,7 @@ bTargetsTiledGPU=true
 bSupports4ComponentUAVReadWrite=true
 bSupportsManualVertexFetch=false
 bSupportsVolumeTextureAtomics=false
+FriendlyName=LOCTEXT("FriendlyShaderPlatformName_IOSES31", "IOS Shader Model 3_1")
 
 [ShaderPlatform METAL_MRT]
 Language=Metal
@@ -67,3 +66,4 @@ bSupportsIndexBufferUAVs=true
 bSupports4ComponentUAVReadWrite=true
 bSupportsManualVertexFetch=false
 bSupportsVolumeTextureAtomics=false
+FriendlyName=LOCTEXT("FriendlyShaderPlatformName_IOSSM5", "IOS Shader Model 5")
\ No newline at end of file
diff --git a/Engine/Config/Mac/DataDrivenPlatformInfo.ini b/Engine/Config/Mac/DataDrivenPlatformInfo.ini
--- a/Engine/Config/Mac/DataDrivenPlatformInfo.ini
+++ b/Engine/Config/Mac/DataDrivenPlatformInfo.ini
@@ -51,6 +51,7 @@ bSupports4ComponentUAVReadWrite=true
 bSupportsVolumeTextureAtomics=false
 bSupportsDualSourceBlending=true
 EnablesHLSL2021ByDefault=1
+FriendlyName=LOCTEXT("FriendlyShaderPlatformName_MacSM5", "Mac Shader Model 5")
 
 [ShaderPlatform METAL_MRT_MAC]
 Language=Metal
@@ -62,6 +63,7 @@ bSupportsIndexBufferUAVs=true
 bSupports4ComponentUAVReadWrite=true
 bSupportsManualVertexFetch=true
 bSupportsVolumeTextureAtomics=false
+FriendlyName=LOCTEXT("FriendlyShaderPlatformName_MacSM5_MRT", "Mac Shader Model 5 MRT")
 
 [ShaderPlatform METAL_MACES3_1]
 Language=Metal
@@ -72,3 +74,4 @@ bIsMobile=true
 bSupports4ComponentUAVReadWrite=true
 bSupportsManualVertexFetch=false
 bSupportsVolumeTextureAtomics=false
+FriendlyName=LOCTEXT("FriendlyShaderPlatformName_MacES31", "Mac Shader Model 3.1")
\ No newline at end of file
diff --git a/Engine/Config/TVOS/DataDrivenPlatformInfo.ini b/Engine/Config/TVOS/DataDrivenPlatformInfo.ini
--- a/Engine/Config/TVOS/DataDrivenPlatformInfo.ini
+++ b/Engine/Config/TVOS/DataDrivenPlatformInfo.ini
@@ -25,6 +25,7 @@ bTargetsTiledGPU=true
 bSupportsIndexBufferUAVs=true
 bSupports4ComponentUAVReadWrite=true
 bSupportsManualVertexFetch=false
+FriendlyName=LOCTEXT("FriendlyShaderPlatformName_TVOS_SM5", "TVOS Shader Model 5")
 
 [ShaderPlatform METAL_TVOS]
 Language=Metal
@@ -34,3 +35,4 @@ bIsMobile=true
 bTargetsTiledGPU=true
 bSupports4ComponentUAVReadWrite=true
 bSupportsManualVertexFetch=false
+FriendlyName=LOCTEXT("FriendlyShaderPlatformName_TVOS_ES31", "TVOS Shader Model 3_1")
\ No newline at end of file
diff --git a/Engine/Config/VulkanPC/DataDrivenPlatformInfo.ini b/Engine/Config/VulkanPC/DataDrivenPlatformInfo.ini
--- a/Engine/Config/VulkanPC/DataDrivenPlatformInfo.ini
+++ b/Engine/Config/VulkanPC/DataDrivenPlatformInfo.ini
@@ -132,6 +132,7 @@ bSupportsROV=false
 bSupportsOIT=false
 
 EnablesHLSL2021ByDefault=1
+FriendlyName=LOCTEXT("FriendlyShaderPlatformName_VulkanPC_SM5", "Vulkan PC Shader Model 5")
 
 bSupportsBindless=true
 
@@ -146,3 +147,4 @@ bSupportsMobileMultiView=true
 bSupportsIndexBufferUAVs=true
 bSupportsManualVertexFetch=false
 bSupportsSceneDataCompressedTransforms=true
+FriendlyName=LOCTEXT("FriendlyShaderPlatformName_VulkanPC_ES31", "Vulkan PC Shader Model 3_1")
\ No newline at end of file
diff --git a/Engine/Config/Windows/DataDrivenPlatformInfo.ini b/Engine/Config/Windows/DataDrivenPlatformInfo.ini
--- a/Engine/Config/Windows/DataDrivenPlatformInfo.ini
+++ b/Engine/Config/Windows/DataDrivenPlatformInfo.ini
@@ -74,16 +74,16 @@ bSupportsIndexBufferUAVs=true
 bSupportsROV=true
 bSupportsOIT=true
 SupportsMultiViewport=RuntimeDependent
+FriendlyName=LOCTEXT("FriendlyShaderPlatformName_D3D_SM5", "D3D Shader Model 5")
 
 [PreviewPlatform PCD3D_SM5]
 PlatformName=PCD3D_SM5
 ShaderFormat=PCD3D_SM5
 ShaderPlatform=PCD3D_SM5
-MenuText=LOCTEXT("PreviewMenuText_PCD3D_SM5", "Shader Model 5")
 MenuTooltip=LOCTEXT("PreviewMenuTooltip_PCD3D_SM5", "DirectX 11, OpenGL 4.3+, PS4, XB1")
 
 [ShaderPlatform PCD3D_SM6]
-FriendlyName=LOCTEXT("FriendlyShaderPlatformName_PCD3D_SM6", "Shader Model 6")
+FriendlyName=LOCTEXT("FriendlyShaderPlatformName_D3D_SM6", "D3D Shader Model 6")
 Language=D3D
 MaxFeatureLevel=SM6
 ShaderFormat=PCD3D_SM6
@@ -137,7 +137,6 @@ SupportsMultiViewport=RuntimeGuaranteed
 PlatformName=PCD3D_SM6
 ShaderFormat=PCD3D_SM6
 ShaderPlatform=PCD3D_SM6
-MenuText=LOCTEXT("PreviewMenuText_PCD3D_SM6", "Shader Model 6")
 MenuTooltip=LOCTEXT("PreviewMenuTooltip_PCD3D_SM6", "Windows using SM6 profile")
 
 [ShaderPlatform PCD3D_ES3_1]
@@ -152,12 +151,12 @@ bSupportsInstancedStereo=true
 bSupportsManualVertexFetch=false
 bSupportsSceneDataCompressedTransforms=true
 bSupportsDebugViewShaders=true
+FriendlyName=LOCTEXT("FriendlyShaderPlatformName_D3D_ES31", "D3D Shader Model 3_1")
 
 [PreviewPlatform PCD3D_ES3_1]
 PlatformName=PCD3D_ES3_1
 ShaderFormat=PCD3D_ES31
 ShaderPlatform=PCD3D_ES3_1
-MenuText=LOCTEXT("PreviewMenuText_PCD3D_ES3_1", "Shader Model 3_1")
 MenuTooltip=LOCTEXT("PreviewMenuTooltip_PCD3D_ES3_1", "Windows using ES3_1 profile.")
 
 [ShaderPlatform OPENGL_PCES3_1]
@@ -168,3 +167,4 @@ MaxFeatureLevel=ES3_1
 ShaderFormat=GLSL_150_ES31
 bIsMobile=true
 bSupportsManualVertexFetch=false
+FriendlyName=LOCTEXT("FriendlyShaderPlatformName_PCGL_ES31", "Opengl PC Shader Model 3.1")
\ No newline at end of file
diff --git a/Engine/Source/Editor/LevelEditor/Private/LevelEditor.cpp b/Engine/Source/Editor/LevelEditor/Private/LevelEditor.cpp
--- a/Engine/Source/Editor/LevelEditor/Private/LevelEditor.cpp
+++ b/Engine/Source/Editor/LevelEditor/Private/LevelEditor.cpp
@@ -1787,7 +1787,8 @@ void FLevelEditorModule::BindGlobalLevelEditorCommands()
 		FIsActionButtonVisible::CreateStatic(FLevelEditorActionCallbacks::IsPreviewModeButtonVisible));
 
 	const TArray<FPreviewPlatformMenuItem>& MenuItems = FDataDrivenPlatformInfoRegistry::GetAllPreviewPlatformMenuItems();
-	check(MenuItems.Num() == Commands.PreviewPlatformOverrides.Num());
+	// We need one extra slot for the Disable Preview option
+	check(MenuItems.Num() + 1 == Commands.PreviewPlatformOverrides.Num());
 
 	for (int32 Index=0; Index < MenuItems.Num(); Index++)
 	{
@@ -1822,6 +1823,16 @@ void FLevelEditorModule::BindGlobalLevelEditorCommands()
 			FIsActionChecked::CreateStatic(&FLevelEditorActionCallbacks::IsPreviewPlatformChecked, PreviewFeatureLevelInfo));
 	}
 
+	// Add the Disable Preview Menu Item
+	{
+		FPreviewPlatformInfo PreviewFeatureLevelInfo(GMaxRHIFeatureLevel, GMaxRHIShaderPlatform, NAME_None, NAME_None, NAME_None, true, NAME_None);
+		ActionList.MapAction(
+			Commands.PreviewPlatformOverrides[MenuItems.Num()],
+			FExecuteAction::CreateStatic(&FLevelEditorActionCallbacks::SetPreviewPlatform, PreviewFeatureLevelInfo),
+			FCanExecuteAction::CreateStatic(&FLevelEditorActionCallbacks::CanExecutePreviewPlatform, PreviewFeatureLevelInfo),
+			FIsActionChecked::CreateStatic(&FLevelEditorActionCallbacks::IsPreviewPlatformChecked, PreviewFeatureLevelInfo));
+	}
+
 	ActionList.MapAction(
 		Commands.OpenMergeActor,
 		FExecuteAction::CreateStatic(&FLevelEditorActionCallbacks::OpenMergeActor_Clicked)
diff --git a/Engine/Source/Editor/LevelEditor/Private/LevelEditorActions.cpp b/Engine/Source/Editor/LevelEditor/Private/LevelEditorActions.cpp
--- a/Engine/Source/Editor/LevelEditor/Private/LevelEditorActions.cpp
+++ b/Engine/Source/Editor/LevelEditor/Private/LevelEditorActions.cpp
@@ -3715,17 +3715,30 @@ void FLevelEditorCommands::RegisterCommands()
 	// Add preview platforms
 	for (const FPreviewPlatformMenuItem& Item : FDataDrivenPlatformInfoRegistry::GetAllPreviewPlatformMenuItems())
 	{
+		EShaderPlatform ShaderPlatform = FDataDrivenShaderPlatformInfo::GetShaderPlatformFromName(Item.ShaderPlatformPreview);
+		FText FriendlyName = FDataDrivenShaderPlatformInfo::GetFriendlyName(ShaderPlatform);
+
 		PreviewPlatformOverrides.Add(
 			FUICommandInfoDecl(
 				this->AsShared(),
 				FName(*FString::Printf(TEXT("PreviewPlatformOverrides_%s_%s"), *Item.PlatformName.ToString(), *Item.ShaderFormat.ToString())),
-				Item.MenuText,
+				FriendlyName,
 				Item.MenuTooltip)
 			.UserInterfaceType(EUserInterfaceActionType::Check)
 			.DefaultChord(FInputChord())
 		);
 	}
 
+	PreviewPlatformOverrides.Add(
+		FUICommandInfoDecl(
+			this->AsShared(),
+			FName(TEXT("Disable Preview Command")),
+			NSLOCTEXT("PreviewPlatform", "PreviewMenuText_DisablePreview", "Disable Preview"),
+			NSLOCTEXT("PreviewPlatform", "PreviewMenuTooltip_DisablePreview", "Disable the current Preview Shader Platform"))
+		.UserInterfaceType(EUserInterfaceActionType::Check)
+		.DefaultChord(FInputChord())
+	);
+
 	UI_COMMAND(OpenMergeActor, "Merge Actors", "Opens the Merge Actor panel", EUserInterfaceActionType::Button, FInputChord());
 }
 
diff --git a/Engine/Source/Editor/LevelEditor/Private/LevelEditorToolBar.cpp b/Engine/Source/Editor/LevelEditor/Private/LevelEditorToolBar.cpp
--- a/Engine/Source/Editor/LevelEditor/Private/LevelEditorToolBar.cpp
+++ b/Engine/Source/Editor/LevelEditor/Private/LevelEditorToolBar.cpp
@@ -53,7 +53,6 @@
 #include "LauncherPlatformModule.h"
 #include "Misc/ScopedSlowTask.h"
 #include "MaterialShaderQualitySettings.h"
-#include "RHIShaderPlatformDefinitions.inl"
 #include "LevelEditorMenuContext.h"
 #include "ToolMenus.h"
 #include "Subsystems/AssetEditorSubsystem.h"
@@ -1415,8 +1414,8 @@ void FLevelEditorToolBar::RegisterLevelEditorToolBar( const TSharedRef<FUIComman
 				EShaderPlatform MaxRHIFeatureLevelPlatform = GetFeatureLevelShaderPlatform(GMaxRHIFeatureLevel);
 
 				{
-					const FText& RenderingAsPlatformName = GetFriendlyShaderPlatformName(GEditor->PreviewPlatform.bPreviewFeatureLevelActive ? PreviewShaderPlatform : MaxRHIFeatureLevelPlatform);
-                    const FText& SwitchToPlatformName = GetFriendlyShaderPlatformName(GEditor->PreviewPlatform.bPreviewFeatureLevelActive ? MaxRHIFeatureLevelPlatform : PreviewShaderPlatform);
+					const FText& RenderingAsPlatformName = FDataDrivenShaderPlatformInfo::GetFriendlyName(GEditor->PreviewPlatform.bPreviewFeatureLevelActive ? PreviewShaderPlatform : MaxRHIFeatureLevelPlatform);
+                    const FText& SwitchToPlatformName = FDataDrivenShaderPlatformInfo::GetFriendlyName(GEditor->PreviewPlatform.bPreviewFeatureLevelActive ? MaxRHIFeatureLevelPlatform : PreviewShaderPlatform);
 					if (PreviewShaderPlatform == MaxRHIFeatureLevelPlatform)
 					{
 						return FText::Format(LOCTEXT("PreviewModeViewingAs", "Viewing {0}."), RenderingAsPlatformName);
diff --git a/Engine/Source/Editor/UnrealEd/Private/PlayLevel.cpp b/Engine/Source/Editor/UnrealEd/Private/PlayLevel.cpp
--- a/Engine/Source/Editor/UnrealEd/Private/PlayLevel.cpp
+++ b/Engine/Source/Editor/UnrealEd/Private/PlayLevel.cpp
@@ -71,7 +71,6 @@
 #include "Slate/SceneViewport.h"
 #include "Kismet2/KismetEditorUtilities.h"
 #include "Kismet2/BlueprintEditorUtils.h"
-#include "RHIShaderPlatformDefinitions.inl"
 
 
 #include "LevelEditor.h"
@@ -3199,7 +3198,7 @@ FText GeneratePIEViewportWindowTitle(const EPlayNetMode InNetMode, const ERHIFea
 	FFormatNamedArguments Args;
 	Args.Add(TEXT("GameName"), FText::FromString(FString(WindowTitleOverride.IsEmpty() ? FApp::GetProjectName() : WindowTitleOverride.ToString())));
 	Args.Add(TEXT("PlatformBits"), FText::FromString(PlatformBitsString));
-	Args.Add(TEXT("RHIName"), GetFriendlyShaderPlatformName(GetFeatureLevelShaderPlatform(InFeatureLevel)));
+	Args.Add(TEXT("RHIName"), FDataDrivenShaderPlatformInfo::GetFriendlyName(GetFeatureLevelShaderPlatform(InFeatureLevel)));
 	
 	if (InNetMode == PIE_Client)
 	{
diff --git a/Engine/Source/Editor/UnrealEd/Private/SEditorViewport.cpp b/Engine/Source/Editor/UnrealEd/Private/SEditorViewport.cpp
--- a/Engine/Source/Editor/UnrealEd/Private/SEditorViewport.cpp
+++ b/Engine/Source/Editor/UnrealEd/Private/SEditorViewport.cpp
@@ -20,7 +20,6 @@
 #include "Widgets/Input/SSpinBox.h"
 #include "Widgets/Input/SCheckBox.h"
 #include "MaterialShaderQualitySettings.h"
-#include "RHIShaderPlatformDefinitions.inl"
 #include "RayTracingDebugVisualizationMenuCommands.h"
 #include "GPUSkinCacheVisualizationMenuCommands.h"
 #include "GPUSkinCache.h"
@@ -861,7 +860,7 @@ FText SEditorViewport::GetCurrentFeatureLevelPreviewText(bool bDrawOnlyLabel) co
 		{
 			ERHIFeatureLevel::Type TargetFeatureLevel = World->FeatureLevel;
 			EShaderPlatform ShaderPlatform = GetShaderPlatformHelper(TargetFeatureLevel);
-			const FText& PlatformText = GetFriendlyShaderPlatformName(ShaderPlatform);
+			const FText& PlatformText = FDataDrivenShaderPlatformInfo::GetFriendlyName(ShaderPlatform);
 			LabelName = FText::Format(LOCTEXT("WorldFeatureLevel", "{0}"), PlatformText);
 		}
 	}
diff --git a/Engine/Source/Runtime/Core/Private/Misc/DataDrivenPlatformInfoRegistry.cpp b/Engine/Source/Runtime/Core/Private/Misc/DataDrivenPlatformInfoRegistry.cpp
--- a/Engine/Source/Runtime/Core/Private/Misc/DataDrivenPlatformInfoRegistry.cpp
+++ b/Engine/Source/Runtime/Core/Private/Misc/DataDrivenPlatformInfoRegistry.cpp
@@ -285,7 +285,6 @@ static void ParsePreviewPlatforms(const FConfigFile& IniFile)
 			Item.InactiveIconName = *GetSectionString(Section.Value, FName("InactiveIconName"));
 			Item.DeviceProfileName = *GetSectionString(Section.Value, FName("DeviceProfileName"));
 			Item.ShaderPlatformPreview = *(GetSectionString(Section.Value, FName("ShaderPlatform")) + TEXT("_PREVIEW"));
-			FTextStringHelper::ReadFromBuffer(*GetSectionString(Section.Value, FName("MenuText")), Item.MenuText);
 			FTextStringHelper::ReadFromBuffer(*GetSectionString(Section.Value, FName("MenuTooltip")), Item.MenuTooltip);
 			FTextStringHelper::ReadFromBuffer(*GetSectionString(Section.Value, FName("IconText")), Item.IconText);
 			PreviewPlatformMenuItems.Add(Item);
diff --git a/Engine/Source/Runtime/Core/Public/Misc/DataDrivenPlatformInfoRegistry.h b/Engine/Source/Runtime/Core/Public/Misc/DataDrivenPlatformInfoRegistry.h
--- a/Engine/Source/Runtime/Core/Public/Misc/DataDrivenPlatformInfoRegistry.h
+++ b/Engine/Source/Runtime/Core/Public/Misc/DataDrivenPlatformInfoRegistry.h
@@ -104,7 +104,6 @@ struct FPreviewPlatformMenuItem
 	FName ActiveIconName;
 	FString InactiveIconPath;
 	FName InactiveIconName;
-	FText MenuText;
 	FText MenuTooltip;
 	FText IconText;
 	FName DeviceProfileName;
diff --git a/Engine/Source/Runtime/RHI/Private/RHI.cpp b/Engine/Source/Runtime/RHI/Private/RHI.cpp
--- a/Engine/Source/Runtime/RHI/Private/RHI.cpp
+++ b/Engine/Source/Runtime/RHI/Private/RHI.cpp
@@ -2670,7 +2670,6 @@ void FGenericDataDrivenShaderPlatformInfo::UpdatePreviewPlatforms()
 
 const EShaderPlatform FGenericDataDrivenShaderPlatformInfo::GetShaderPlatformFromName(const FName ShaderPlatformName)
 {
-	EShaderPlatform ShaderPlatformToPreview = EShaderPlatform::SP_NumPlatforms;
 	for (int32 i = 0; i < SP_NumPlatforms; ++i)
 	{
 		const EShaderPlatform Platform = static_cast<EShaderPlatform>(i);
@@ -2681,11 +2680,10 @@ const EShaderPlatform FGenericDataDrivenShaderPlatformInfo::GetShaderPlatformFro
 
 		if (Infos[Platform].Name == ShaderPlatformName)
 		{
-			ShaderPlatformToPreview = Platform;
-			break;
+			return Platform;
 		}
 	}
-	return ShaderPlatformToPreview;
+	return SP_NumPlatforms;
 }
 
 //
diff --git a/Engine/Source/Runtime/RHI/Public/RHIShaderPlatformDefinitions.inl b/Engine/Source/Runtime/RHI/Public/RHIShaderPlatformDefinitions.inl
deleted file mode 100644
