Movement of configuration to base files to avoid overriding project settings

diff --git a/Engine/Config/Android/AndroidEngine.ini b/Engine/Config/Android/AndroidEngine.ini
--- a/Engine/Config/Android/AndroidEngine.ini
+++ b/Engine/Config/Android/AndroidEngine.ini
@@ -40,9 +40,6 @@ PreviewDeviceProfileSelectionModule="AndroidDeviceProfileSelector"
 [Advertising]
 DefaultProviderName=AndroidAdvertising
 
-[OnlineSubsystem]
-DefaultPlatformService=GooglePlay
-
 [SlateRenderer]
 NumPreallocatedVertices=200
 
@@ -94,10 +91,6 @@ r.TemporalAA.R11G11B10History=0
 ; Do not backtrace vulkan memory allocations
 r.Vulkan.MemoryBacktrace=0
 
-[OnlineSubsystemGooglePlay.Store]
-bSupportsInAppPurchasing=true
-
-
 [/Script/Engine.RendererSettings]
 r.DistanceFields=0
 r.TemporalAA.Mobile.UseCompute=0
diff --git a/Engine/Config/Android/BaseAndroidEngine.ini b/Engine/Config/Android/BaseAndroidEngine.ini
--- a/Engine/Config/Android/BaseAndroidEngine.ini
+++ b/Engine/Config/Android/BaseAndroidEngine.ini
@@ -1,7 +1,11 @@
 [OnlineSubsystem]
+DefaultPlatformService=GooglePlay
 NativePlatformService=GooglePlay
 LocalPlatformName=AND
 
+[OnlineSubsystemGooglePlay.Store]
+bSupportsInAppPurchasing=true
+
 [Audio]
 AudioCaptureModuleName=AudioCaptureAndroid
 
diff --git a/Engine/Config/IOS/BaseIOSEngine.ini b/Engine/Config/IOS/BaseIOSEngine.ini
--- a/Engine/Config/IOS/BaseIOSEngine.ini
+++ b/Engine/Config/IOS/BaseIOSEngine.ini
@@ -1,7 +1,11 @@
 [OnlineSubsystem]
+DefaultPlatformService=IOS
 NativePlatformService=IOS
 LocalPlatformName=IOS
 
+[OnlineSubsystemIOS.Store]
+bSupportsInAppPurchasing=true
+
 [Audio]
 AudioCaptureModuleName=AudioCaptureAudioUnit
 
diff --git a/Engine/Config/IOS/IOSEngine.ini b/Engine/Config/IOS/IOSEngine.ini
--- a/Engine/Config/IOS/IOSEngine.ini
+++ b/Engine/Config/IOS/IOSEngine.ini
@@ -32,18 +32,12 @@ DefaultMemoryBucket_MinGB=4
 SmallerMemoryBucket_MinGB=3
 SmallestMemoryBucket_MinGB=2
 
-[OnlineSubsystem]
-DefaultPlatformService=IOS
-
 [Advertising]
 DefaultProviderName=IOSAdvertising
 
 [DeviceProfileManager]
 DeviceProfileSelectionModule="IOSDeviceProfileSelector"
 
-[OnlineSubsystemIOS.Store]
-bSupportsInAppPurchasing=true
-
 [SlateRenderer]
 NumPreallocatedVertices=200
 
diff --git a/Engine/Config/Linux/BaseLinuxEngine.ini b/Engine/Config/Linux/BaseLinuxEngine.ini
--- a/Engine/Config/Linux/BaseLinuxEngine.ini
+++ b/Engine/Config/Linux/BaseLinuxEngine.ini
@@ -1,3 +1,6 @@
+[OnlineSubsystem]
+LocalPlatformName=LNX
+
 [PlatformPaths]
 UnrealBuildTool=Engine/Build/BatchFiles/Linux/Build.sh
 
diff --git a/Engine/Config/Linux/LinuxEngine.ini b/Engine/Config/Linux/LinuxEngine.ini
--- a/Engine/Config/Linux/LinuxEngine.ini
+++ b/Engine/Config/Linux/LinuxEngine.ini
@@ -28,8 +28,5 @@ g.TimeoutForBlockOnRenderFence=60000
 [PlatformCrypto]
 PlatformRequiresDataCrypto=True
 
-[OnlineSubsystem]
-LocalPlatformName=LNX
-
 [CrashReportClient]
 bAgreeToCrashUpload=false
