Bump ShadowQuality to 2 for a High bucket on Android/iOS to match quality between forward and deferred shading

diff --git a/Engine/Config/Android/AndroidScalability.ini b/Engine/Config/Android/AndroidScalability.ini
--- a/Engine/Config/Android/AndroidScalability.ini
+++ b/Engine/Config/Android/AndroidScalability.ini
@@ -60,7 +60,7 @@ r.Shadow.CSM.MaxCascades=1
 r.Shadow.RadiusThreshold=0.06
 
 [ShadowQuality@2]
-r.ShadowQuality=1
+r.ShadowQuality=2
 r.Shadow.MaxCSMResolution=2048
 r.Shadow.DistanceScale=1.0
 r.Shadow.CSM.MaxCascades=2
diff --git a/Engine/Config/IOS/IOSScalability.ini b/Engine/Config/IOS/IOSScalability.ini
--- a/Engine/Config/IOS/IOSScalability.ini
+++ b/Engine/Config/IOS/IOSScalability.ini
@@ -60,7 +60,7 @@ r.Shadow.CSM.MaxCascades=1
 r.Shadow.RadiusThreshold=0.06
 
 [ShadowQuality@2]
-r.ShadowQuality=1
+r.ShadowQuality=2
 r.Shadow.MaxCSMResolution=2048
 r.Shadow.DistanceScale=1.0
 r.Shadow.CSM.MaxCascades=2
