Mobile: Set maximum tiles produced per-frame to 1-2 for RVT on iOS and Android

diff --git a/Engine/Config/Android/AndroidScalability.ini b/Engine/Config/Android/AndroidScalability.ini
--- a/Engine/Config/Android/AndroidScalability.ini
+++ b/Engine/Config/Android/AndroidScalability.ini
@@ -187,15 +187,19 @@ fx.Niagara.QualityLevel=1
 
 [TextureQuality@0]
 r.MaxAnisotropy=1
+r.VT.MaxTilesProducedPerFrame=1
 
 [TextureQuality@1]
 r.MaxAnisotropy=1
+r.VT.MaxTilesProducedPerFrame=1
 
 [TextureQuality@2]
 r.MaxAnisotropy=2
+r.VT.MaxTilesProducedPerFrame=1
 
 [TextureQuality@3]
 r.MaxAnisotropy=4
+r.VT.MaxTilesProducedPerFrame=2
 
 [ScalabilityPreview]
 ; don't want to preview any texture streaming changes in the editor
diff --git a/Engine/Config/IOS/IOSScalability.ini b/Engine/Config/IOS/IOSScalability.ini
--- a/Engine/Config/IOS/IOSScalability.ini
+++ b/Engine/Config/IOS/IOSScalability.ini
@@ -187,15 +187,19 @@ fx.Niagara.QualityLevel=1
 
 [TextureQuality@0]
 r.MaxAnisotropy=1
+r.VT.MaxTilesProducedPerFrame=1
 
 [TextureQuality@1]
 r.MaxAnisotropy=1
+r.VT.MaxTilesProducedPerFrame=1
 
 [TextureQuality@2]
 r.MaxAnisotropy=2
+r.VT.MaxTilesProducedPerFrame=1
 
 [TextureQuality@3]
 r.MaxAnisotropy=4
+r.VT.MaxTilesProducedPerFrame=2
 
 [ScalabilityPreview]
 ; don't want to preview any texture streaming changes in the editor
