Adding screen density for Nothing Phone as OEM values are lower than expected

[RN] Adding screen density for Nothing Phone

diff --git a/Engine/Config/Android/AndroidEngine.ini b/Engine/Config/Android/AndroidEngine.ini
--- a/Engine/Config/Android/AndroidEngine.ini
+++ b/Engine/Config/Android/AndroidEngine.ini
@@ -65,6 +65,7 @@ DefaultPlatformService=AndroidLocalNotification
 +Devices=(Model="Pixel XL",Density=534)
 +Devices=(Model="HTC 10",Density=564)
 +Devices=(Model="EVA-L09",Density=423) ; Huawei P9
++Devices=(Model="A063", Density=420) ; Nothing Phone
 
 
 ;GalaxyS5=432
