Do not force D24_S8 depth target on Android Vulkan

diff --git a/Engine/Config/Android/AndroidEngine.ini b/Engine/Config/Android/AndroidEngine.ini
--- a/Engine/Config/Android/AndroidEngine.ini
+++ b/Engine/Config/Android/AndroidEngine.ini
@@ -77,8 +77,6 @@ r.VSync=1
 r.ParallelShadows=0
 ; high priority thread tasks can be downgraded to normal priority if a task is created with EPowerSavingEligibility::Eligible.
 TaskGraph.EnablePowerSavingThreadPriorityReduction=1
-; prefer D24_S8 for depth/stencil
-r.Vulkan.Depth24Bit=1
 ;GalaxyS6Edge=577
 r.Shadow.ForceSerialSingleRenderPass=1
 ; Support Mali 64k texel buffer limitation, tuned for 1080p. If you target 2k-4k resolutions please also set the LightGridPixelSize to 128
