Enable landscape GPU culling on Android and iOS

diff --git a/Engine/Config/Android/AndroidEngine.ini b/Engine/Config/Android/AndroidEngine.ini
--- a/Engine/Config/Android/AndroidEngine.ini
+++ b/Engine/Config/Android/AndroidEngine.ini
@@ -95,6 +95,8 @@ r.Vulkan.MemoryBacktrace=0
 [/Script/Engine.RendererSettings]
 r.DistanceFields=0
 r.TemporalAA.Mobile.UseCompute=0
+; Enable GPU culling for landscape
+landscape.SupportGPUCulling=1
 
 [AlternateTextureCompression]
 ;disabling TextureCompressionFormat for Android/IOS no longer needed
diff --git a/Engine/Config/IOS/IOSEngine.ini b/Engine/Config/IOS/IOSEngine.ini
--- a/Engine/Config/IOS/IOSEngine.ini
+++ b/Engine/Config/IOS/IOSEngine.ini
@@ -85,6 +85,8 @@ fx.Niagara.QualityLevel.Max=1
 [/Script/Engine.RendererSettings]
 r.DistanceFields=0
 r.TemporalAA.Mobile.UseCompute=0
+; Enable GPU culling for landscape
+landscape.SupportGPUCulling=1
 
 [AlternateTextureCompression]
 ;disabling TextureCompressionFormat for Android/IOS no longer needed
