Enabled Distance Fields in Android rendering settings, on ES31 platforms this should still be disabled by the device profiles and cvar

diff --git a/Engine/Config/Android/AndroidEngine.ini b/Engine/Config/Android/AndroidEngine.ini
--- a/Engine/Config/Android/AndroidEngine.ini
+++ b/Engine/Config/Android/AndroidEngine.ini
@@ -91,7 +91,7 @@ r.TemporalAA.R11G11B10History=0
 r.Vulkan.MemoryBacktrace=0
 
 [/Script/Engine.RendererSettings]
-r.DistanceFields=0
+r.DistanceFields=1
 r.TemporalAA.Mobile.UseCompute=0
 ; Enable GPU culling for landscape
 landscape.SupportGPUCulling=1
