Add bSupportsVariableRateShading to Android DDPI.

diff --git a/Engine/Config/Android/DataDrivenPlatformInfo.ini b/Engine/Config/Android/DataDrivenPlatformInfo.ini
--- a/Engine/Config/Android/DataDrivenPlatformInfo.ini
+++ b/Engine/Config/Android/DataDrivenPlatformInfo.ini
@@ -83,6 +83,7 @@ bSupportsIndexBufferUAVs=true
 bSupportsManualVertexFetch=false
 bSupportsClipDistance = false
 bSupportsShaderPipelines = false
+bSupportsVariableRateShading = true
 bSupportsVertexShaderSRVs=false
 bSupportsUniformBufferObjects = true
 FriendlyName=LOCTEXT("FriendlyShaderPlatformName_AndroidVulkan", "Android Vulkan Mobile")
@@ -111,3 +112,4 @@ bSupportsGPUScene=true
 bSupportsLumenGI=true
 bSupportsDistanceFields=true
 bSupportsDualSourceBlending=false
+bSupportsVariableRateShading = false
\ No newline at end of file
