Set closure count to 1 by default on all low end platforms.

diff --git a/Engine/Config/Android/AndroidEngine.ini b/Engine/Config/Android/AndroidEngine.ini
--- a/Engine/Config/Android/AndroidEngine.ini
+++ b/Engine/Config/Android/AndroidEngine.ini
@@ -99,6 +99,7 @@ r.Substrate.TileCoord8bits=1
 r.Substrate.SheenQuality=2
 r.Substrate.Glints=0
 r.Substrate.SpecularProfile=0
+r.Substrate.ClosuresPerPixel=1
 ; Enable GPU culling for landscape
 landscape.SupportGPUCulling=1
 
diff --git a/Engine/Config/IOS/IOSEngine.ini b/Engine/Config/IOS/IOSEngine.ini
--- a/Engine/Config/IOS/IOSEngine.ini
+++ b/Engine/Config/IOS/IOSEngine.ini
@@ -86,6 +86,7 @@ r.Substrate.TileCoord8bits=1
 r.Substrate.SheenQuality=2
 r.Substrate.Glints=0
 r.Substrate.SpecularProfile=0
+r.Substrate.ClosuresPerPixel=1
 ; Enable GPU culling for landscape
 landscape.SupportGPUCulling=1
 
