Disable r.ContactShadows on mobile platforms by default.

diff --git a/Engine/Config/Android/AndroidScalability.ini b/Engine/Config/Android/AndroidScalability.ini
--- a/Engine/Config/Android/AndroidScalability.ini
+++ b/Engine/Config/Android/AndroidScalability.ini
@@ -52,6 +52,7 @@ r.Shadow.DistanceScale=0.5
 r.Shadow.CSM.MaxCascades=1
 r.Shadow.RadiusThreshold=0.06
 r.VolumetricFog=0
+r.ContactShadows=0
 
 [ShadowQuality@1]
 r.ShadowQuality=1
@@ -60,6 +61,7 @@ r.Shadow.DistanceScale=0.7
 r.Shadow.CSM.MaxCascades=1
 r.Shadow.RadiusThreshold=0.06
 r.VolumetricFog=0
+r.ContactShadows=0
 
 [ShadowQuality@2]
 r.ShadowQuality=2
@@ -68,6 +70,7 @@ r.Shadow.DistanceScale=1.0
 r.Shadow.CSM.MaxCascades=2
 r.Shadow.RadiusThreshold=0.06
 r.VolumetricFog=0
+r.ContactShadows=0
 
 [ShadowQuality@3]
 r.ShadowQuality=2
@@ -76,6 +79,7 @@ r.Shadow.DistanceScale=1.0
 r.Shadow.CSM.MaxCascades=3
 r.Shadow.RadiusThreshold=0.06
 r.VolumetricFog=0
+r.ContactShadows=0
 
 ; ------------------------------------------------------------------------------
 ; Postprocess settings
diff --git a/Engine/Config/IOS/IOSScalability.ini b/Engine/Config/IOS/IOSScalability.ini
--- a/Engine/Config/IOS/IOSScalability.ini
+++ b/Engine/Config/IOS/IOSScalability.ini
@@ -52,6 +52,7 @@ r.Shadow.DistanceScale=0.5
 r.Shadow.CSM.MaxCascades=1
 r.Shadow.RadiusThreshold=0.06
 r.VolumetricFog=0
+r.ContactShadows=0
 
 [ShadowQuality@1]
 r.ShadowQuality=1
@@ -60,6 +61,7 @@ r.Shadow.DistanceScale=0.7
 r.Shadow.CSM.MaxCascades=1
 r.Shadow.RadiusThreshold=0.06
 r.VolumetricFog=0
+r.ContactShadows=0
 
 [ShadowQuality@2]
 r.ShadowQuality=2
@@ -68,6 +70,7 @@ r.Shadow.DistanceScale=1.0
 r.Shadow.CSM.MaxCascades=2
 r.Shadow.RadiusThreshold=0.06
 r.VolumetricFog=0
+r.ContactShadows=0
 
 [ShadowQuality@3]
 r.ShadowQuality=2
@@ -76,6 +79,7 @@ r.Shadow.DistanceScale=1.0
 r.Shadow.CSM.MaxCascades=3
 r.Shadow.RadiusThreshold=0.06
 r.VolumetricFog=0
+r.ContactShadows=0
 
 ; ------------------------------------------------------------------------------
 ; Postprocess settings
