Android: Set max allowed NDK version to 26.

diff --git a/Engine/Config/Android/Android_SDK.json b/Engine/Config/Android/Android_SDK.json
--- a/Engine/Config/Android/Android_SDK.json
+++ b/Engine/Config/Android/Android_SDK.json
@@ -3,5 +3,5 @@
 "AutoSDKDirectory" : "-26",
 
 "MinVersion" : "r25b",
-"MaxVersion" : "r25b",
+"MaxVersion" : "r26",
 }
