Disable hair binding on mobile, as it fails an assert due to a missing SRV on the target mesh

#changelist validated
#virtualized

diff --git a/Engine/Config/Android/AndroidEngine.ini b/Engine/Config/Android/AndroidEngine.ini
--- a/Engine/Config/Android/AndroidEngine.ini
+++ b/Engine/Config/Android/AndroidEngine.ini
@@ -84,7 +84,7 @@ r.Forward.LightGridSizeZ = 8
 r.Forward.MaxCulledLightsPerCell = 8
 r.HairStrands.Simulation=0
 r.HairStrands.Strands=0
-r.HairStrands.Binding=1
+r.HairStrands.Binding=0
 ; R11G11B10 UAV is not supported on Android
 r.TemporalAA.R11G11B10History=0
 ; Do not backtrace vulkan memory allocations
diff --git a/Engine/Config/IOS/BaseIOSEngine.ini b/Engine/Config/IOS/BaseIOSEngine.ini
--- a/Engine/Config/IOS/BaseIOSEngine.ini
+++ b/Engine/Config/IOS/BaseIOSEngine.ini
@@ -54,7 +54,7 @@ RetryResumeDataLimit=3;
 Slate.CacheRenderData=0
 r.HairStrands.Simulation=0
 r.HairStrands.Strands=0
-r.HairStrands.Binding=1
+r.HairStrands.Binding=0
 
 [DeviceProfileManager]
 DeviceProfileSelectionModule="IOSDeviceProfileSelector"
