Enable using NDK 26.2.x

diff --git a/Engine/Config/Android/Android_SDK.json b/Engine/Config/Android/Android_SDK.json
--- a/Engine/Config/Android/Android_SDK.json
+++ b/Engine/Config/Android/Android_SDK.json
@@ -3,5 +3,5 @@
 "AutoSDKDirectory" : "-26",
 
 "MinVersion" : "r25b",
-"MaxVersion" : "r26b"
+"MaxVersion" : "r26c"
 }
