Change r.NeverOcclusionTestDistance from 2000 to 200 for a base engine IOS and Android, so it works for most types of content

diff --git a/Engine/Config/Android/AndroidScalability.ini b/Engine/Config/Android/AndroidScalability.ini
--- a/Engine/Config/Android/AndroidScalability.ini
+++ b/Engine/Config/Android/AndroidScalability.ini
@@ -27,19 +27,19 @@ foliage.DensityScale=1.0
 
 [ViewDistanceQuality@0]
 r.ViewDistanceScale=0.70
-r.NeverOcclusionTestDistance=2000
+r.NeverOcclusionTestDistance=200
 
 [ViewDistanceQuality@1]
 r.ViewDistanceScale=1.0
-r.NeverOcclusionTestDistance=2000
+r.NeverOcclusionTestDistance=200
 
 [ViewDistanceQuality@2]
 r.ViewDistanceScale=1.0
-r.NeverOcclusionTestDistance=2000
+r.NeverOcclusionTestDistance=200
 
 [ViewDistanceQuality@3]
 r.ViewDistanceScale=1.0
-r.NeverOcclusionTestDistance=2000
+r.NeverOcclusionTestDistance=200
 
 ; ------------------------------------------------------------------------------
 ; Shadow settings
diff --git a/Engine/Config/IOS/IOSScalability.ini b/Engine/Config/IOS/IOSScalability.ini
--- a/Engine/Config/IOS/IOSScalability.ini
+++ b/Engine/Config/IOS/IOSScalability.ini
@@ -27,19 +27,19 @@ foliage.DensityScale=1.0
 
 [ViewDistanceQuality@0]
 r.ViewDistanceScale=0.70
-r.NeverOcclusionTestDistance=2000
+r.NeverOcclusionTestDistance=200
 
 [ViewDistanceQuality@1]
 r.ViewDistanceScale=1.0
-r.NeverOcclusionTestDistance=2000
+r.NeverOcclusionTestDistance=200
 
 [ViewDistanceQuality@2]
 r.ViewDistanceScale=1.0
-r.NeverOcclusionTestDistance=2000
+r.NeverOcclusionTestDistance=200
 
 [ViewDistanceQuality@3]
 r.ViewDistanceScale=1.0
-r.NeverOcclusionTestDistance=2000
+r.NeverOcclusionTestDistance=200
 
 ; ------------------------------------------------------------------------------
 ; Shadow settings
