Remove landscape.SupportGPUCulling=1 from platform ini files, since feature is enabled by default

diff --git a/Engine/Config/Android/AndroidEngine.ini b/Engine/Config/Android/AndroidEngine.ini
--- a/Engine/Config/Android/AndroidEngine.ini
+++ b/Engine/Config/Android/AndroidEngine.ini
@@ -100,8 +100,6 @@ r.Substrate.SheenQuality=2
 r.Substrate.Glints=0
 r.Substrate.SpecularProfile=0
 r.Substrate.ClosuresPerPixel=1
-; Enable GPU culling for landscape
-landscape.SupportGPUCulling=1
 ; Disables parallel GDME by default due to not enough threads
 r.Visibility.DynamicMeshElements.Parallel=0
 
diff --git a/Engine/Config/IOS/BaseIOSEngine.ini b/Engine/Config/IOS/BaseIOSEngine.ini
--- a/Engine/Config/IOS/BaseIOSEngine.ini
+++ b/Engine/Config/IOS/BaseIOSEngine.ini
@@ -18,8 +18,6 @@ r.Substrate.SheenQuality=2
 r.Substrate.Glints=0
 r.Substrate.SpecularProfile=0
 r.Substrate.ClosuresPerPixel=1
-; Enable GPU culling for landscape
-landscape.SupportGPUCulling=1
 
 [Advertising]
 DefaultProviderName=IOSAdvertising
