Group the PreviewPlatforms into submenus organized by PlatformName

diff --git a/Engine/Config/Android/DataDrivenPlatformInfo.ini b/Engine/Config/Android/DataDrivenPlatformInfo.ini
--- a/Engine/Config/Android/DataDrivenPlatformInfo.ini
+++ b/Engine/Config/Android/DataDrivenPlatformInfo.ini
@@ -29,7 +29,9 @@ ActiveIconName=LevelEditor.PreviewMode.AndroidES31.Enabled
 InactiveIconName=LevelEditor.PreviewMode.AndroidES31.Disabled
 MenuTooltip=NSLOCTEXT("PreviewPlatform", "PreviewMenuTooltip_AndroidOpenGL ", "Mobile preview using Android OpenGL quality settings.")
 IconText=NSLOCTEXT("PreviewPlatform", "PreviewIconText_AndroidOpenGL", "Android OpenGL")
-DeviceProfileName=Android_Preview_OpenGL
+DeviceProfileName=Android_Preview_OpenGL:Android_Low:Android_Mid:Android_High
+FriendlyName=NSLOCTEXT("PreviewPlatform", "PreviewFriendlyName_GenericOpenglAndroid", "Generic Android OpenGL"):NSLOCTEXT("PreviewPlatform", "PreviewFriendlyName_AndroidOpenGLLow","Android OpenGL Low"):NSLOCTEXT("PreviewPlatform", "PreviewFriendlyName_AndroidOpenGLMid","Android OpenGL Mid"):NSLOCTEXT("PreviewPlatform", "PreviewFriendlyName_AndroidOpenGLHigh","Android OpenGL High")
+
 
 [PreviewPlatform AndroidVulkan]
 EnabledCVar=ini:Engine:Android:/Script/AndroidRuntimeSettings.AndroidRuntimeSettings:bSupportsVulkan
@@ -40,7 +42,9 @@ ActiveIconName=LevelEditor.PreviewMode.AndroidVulkan.Enabled
 InactiveIconName=LevelEditor.PreviewMode.AndroidVulkan.Disabled
 MenuTooltip=NSLOCTEXT("PreviewPlatform", "PreviewMenuTooltip_AndroidVulkan", "Mobile preview using Android Vulkan quality settings.")
 IconText=NSLOCTEXT("PreviewPlatform", "PreviewIconText_AndroidVulkan", "Android Vulkan")
-DeviceProfileName=Android_Preview_Vulkan
+DeviceProfileName=Android_Preview_Vulkan:Android_Low:Android_Mid:Android_High
+FriendlyName=NSLOCTEXT("PreviewPlatform", "PreviewFriendlyName_GenericVulkanAndroid", "Generic Android Vulkan"):NSLOCTEXT("PreviewPlatform", "PreviewFriendlyName_AndroidVulkanLow","Android Vulkan Low"):NSLOCTEXT("PreviewPlatform", "PreviewFriendlyName_AndroidVulkanMid","Android Vulkan Mid"):NSLOCTEXT("PreviewPlatform", "PreviewFriendlyName_AndroidVulkanHigh","Android Vulkan High")
+
 
 [PreviewPlatform AndroidVulkanSM5]
 EnabledCVar=ini:Engine:Android:/Script/AndroidRuntimeSettings.AndroidRuntimeSettings:bSupportsVulkanSM5
diff --git a/Engine/Config/IOS/DataDrivenPlatformInfo.ini b/Engine/Config/IOS/DataDrivenPlatformInfo.ini
--- a/Engine/Config/IOS/DataDrivenPlatformInfo.ini
+++ b/Engine/Config/IOS/DataDrivenPlatformInfo.ini
@@ -33,7 +33,7 @@ InactiveIconName=LevelEditor.PreviewMode.iOS.Disabled
 MenuTooltip=NSLOCTEXT("PreviewPlatform", "PreviewMenuTooltip_iOS", "Mobile preview using iOS material quality settings.")
 IconText=NSLOCTEXT("PreviewPlatform", "PreviewIconText_iOS", "iOS")
 DeviceProfileName=IOS:iPhone6S:iPadPro6_129
-FriendlyName=Generic IOS:iPhone 6S:iPad Pro 6 12.9in
+FriendlyName=NSLOCTEXT("PreviewPlatform", "PreviewFriendlyName_GenericIOS", "Generic IOS"):NSLOCTEXT("PreviewPlatform", "PreviewFriendlyName_iPhone6S", "iPhone 6S"):NSLOCTEXT("PreviewPlatform", "PreviewFriendlyName_iPadPro612.9in", "iPad Pro 6 12.9in")
 
 [PreviewPlatform IOSMetalSM5]
 PlatformName=IOS
diff --git a/Engine/Config/VulkanPC/DataDrivenPlatformInfo.ini b/Engine/Config/VulkanPC/DataDrivenPlatformInfo.ini
--- a/Engine/Config/VulkanPC/DataDrivenPlatformInfo.ini
+++ b/Engine/Config/VulkanPC/DataDrivenPlatformInfo.ini
@@ -287,7 +287,7 @@ bSupportsIndependentSamplers=true
 FriendlyName=LOCTEXT("FriendlyShaderPlatformName_VulkanPC_Mobile", "PC Vulkan Mobile")
 
 [PreviewPlatform VULKAN_SM5]
-PlatformName=VULKAN_SM5
+PlatformName=PC
 ShaderFormat=SF_VULKAN_SM5
 ShaderPlatform=VULKAN_SM5
 MenuTooltip=LOCTEXT("PreviewMenuTooltip_VulkanPC_SM5", "Linux using SM5 profile")
diff --git a/Engine/Config/Windows/DataDrivenPlatformInfo.ini b/Engine/Config/Windows/DataDrivenPlatformInfo.ini
--- a/Engine/Config/Windows/DataDrivenPlatformInfo.ini
+++ b/Engine/Config/Windows/DataDrivenPlatformInfo.ini
@@ -79,7 +79,7 @@ bSupportsIndependentSamplers=true
 bSupportsVertexShaderUAVs=RuntimeGuaranteed
 
 [PreviewPlatform PCD3D_SM5]
-PlatformName=PCD3D_SM5
+PlatformName=PC
 ShaderFormat=PCD3D_SM5
 ShaderPlatform=PCD3D_SM5
 MenuTooltip=LOCTEXT("PreviewMenuTooltip_PCD3D_SM5", "DirectX 11, OpenGL 4.3+, PS4, XB1")
@@ -145,7 +145,7 @@ bSupportsWorkGraphs=true
 bSupportsVertexShaderUAVs=RuntimeGuaranteed
 
 [PreviewPlatform PCD3D_SM6]
-PlatformName=PCD3D_SM6
+PlatformName=PC
 ShaderFormat=PCD3D_SM6
 ShaderPlatform=PCD3D_SM6
 MenuTooltip=LOCTEXT("PreviewMenuTooltip_PCD3D_SM6", "Windows using SM6 profile")
@@ -171,7 +171,7 @@ FriendlyName=LOCTEXT("FriendlyShaderPlatformName_D3D_Mobile", "PC D3D Mobile")
 SupportsMultiViewport=RuntimeDependent
 
 [PreviewPlatform PCD3D_ES3_1]
-PlatformName=PCD3D_ES3_1
+PlatformName=PC
 ShaderFormat=PCD3D_ES31
 ShaderPlatform=PCD3D_ES3_1
 MenuTooltip=LOCTEXT("PreviewMenuTooltip_PCD3D_Mobile", "Windows using Mobile profile.")
diff --git a/Engine/Source/Editor/LevelEditor/Private/LevelEditorActions.cpp b/Engine/Source/Editor/LevelEditor/Private/LevelEditorActions.cpp
--- a/Engine/Source/Editor/LevelEditor/Private/LevelEditorActions.cpp
+++ b/Engine/Source/Editor/LevelEditor/Private/LevelEditorActions.cpp
@@ -4027,8 +4027,12 @@ void FLevelEditorCommands::RegisterCommands()
 			.UserInterfaceType(EUserInterfaceActionType::Check)
 			.DefaultChord(FInputChord())
 		);
+
+		PlatformToPreviewPlatformOverrides.FindOrAdd(Item.PlatformName).Add(PreviewPlatformOverrides.Last());
 	}
 
+	PlatformToPreviewPlatformOverrides.KeyStableSort([](FName lhs, FName rhs) {return lhs.Compare(rhs) < 0; });
+
 	UI_COMMAND(OpenMergeActor, "Merge Actors", "Opens the Merge Actor panel", EUserInterfaceActionType::Button, FInputChord());
 }
 
diff --git a/Engine/Source/Editor/LevelEditor/Private/ViewportToolbar/LevelEditorViewportToolbarSections.cpp b/Engine/Source/Editor/LevelEditor/Private/ViewportToolbar/LevelEditorViewportToolbarSections.cpp
--- a/Engine/Source/Editor/LevelEditor/Private/ViewportToolbar/LevelEditorViewportToolbarSections.cpp
+++ b/Engine/Source/Editor/LevelEditor/Private/ViewportToolbar/LevelEditorViewportToolbarSections.cpp
@@ -1300,11 +1300,24 @@ FToolMenuEntry CreateFeatureLevelPreviewSubmenu()
 			[](UToolMenu* InMenu) -> void
 			{
 				FToolMenuSection& Section =
-					InMenu->AddSection("EditorPreviewMode", LOCTEXT("EditorPreviewModeDevices", "Preview Devices"));
-				// Preview platforms discovered from ITargetPlatforms.
-				for (auto& Item : FLevelEditorCommands::Get().PreviewPlatformOverrides)
+					InMenu->AddSection("EditorPreviewMode", LOCTEXT("EditorPreviewModeDevices", "Preview Platforms"));
+
+				for (auto Iter = FLevelEditorCommands::Get().PlatformToPreviewPlatformOverrides.CreateConstIterator(); Iter; ++Iter)
 				{
-					Section.AddMenuEntry(Item);
+					FName PlatformName = Iter.Key();
+					const TArray<TSharedPtr<FUICommandInfo>>& CommandList = Iter.Value();
+
+					Section.AddSubMenu(FName(PlatformName), FText::FromString(PlatformName.ToString()), FText(),
+						FNewToolMenuDelegate::CreateLambda([CommandList](UToolMenu* InSubMenu)
+							{
+								FToolMenuSection& Section =
+									InSubMenu->AddSection("", LOCTEXT("EditorPreviewModeDevices", "Preview Devices"));
+								for (const TSharedPtr<FUICommandInfo>& Command : CommandList)
+								{
+									Section.AddMenuEntry(Command);
+								}
+							})
+					);
 				}
 			}
 		)
diff --git a/Engine/Source/Editor/LevelEditor/Public/LevelEditorActions.h b/Engine/Source/Editor/LevelEditor/Public/LevelEditorActions.h
--- a/Engine/Source/Editor/LevelEditor/Public/LevelEditorActions.h
+++ b/Engine/Source/Editor/LevelEditor/Public/LevelEditorActions.h
@@ -610,7 +610,8 @@ public:
 	TSharedPtr< FUICommandInfo > ToggleFeatureLevelPreview;
 
 	TArray<TSharedPtr<FUICommandInfo>> PreviewPlatformOverrides;
-	
+	TMap<FName, TArray<TSharedPtr<FUICommandInfo>>> PlatformToPreviewPlatformOverrides;
+
 	///**
 	// * Mode Commands                   
 	// */
diff --git a/Engine/Source/Runtime/Core/Private/Misc/DataDrivenPlatformInfoRegistry.cpp b/Engine/Source/Runtime/Core/Private/Misc/DataDrivenPlatformInfoRegistry.cpp
--- a/Engine/Source/Runtime/Core/Private/Misc/DataDrivenPlatformInfoRegistry.cpp
+++ b/Engine/Source/Runtime/Core/Private/Misc/DataDrivenPlatformInfoRegistry.cpp
@@ -299,11 +299,10 @@ static void ParsePreviewPlatforms(const FConfigFile& IniFile)
 
 
 				FString AllDeviceProfiles = GetSectionString(Section.Value, FName("DeviceProfileName"));
-				FText AllFriendlyName;
-				FTextStringHelper::ReadFromBuffer(*GetSectionString(Section.Value, FName("FriendlyName")), AllFriendlyName);
+				FString AllFriendlyName = GetSectionString(Section.Value, FName("FriendlyName"));
 				TArray<FString> DeviceProfileNames, FriendlyNames;
 				AllDeviceProfiles.ParseIntoArray(DeviceProfileNames, TEXT(":"));
-				AllFriendlyName.ToString().ParseIntoArray(FriendlyNames, TEXT(":"));
+				AllFriendlyName.ParseIntoArray(FriendlyNames, TEXT(":"));
 
 				if (DeviceProfileNames.Num() == 0)
 				{
@@ -316,7 +315,7 @@ static void ParsePreviewPlatforms(const FConfigFile& IniFile)
 					Item.DeviceProfileName = *DeviceProfileNames[DPIndex].TrimStartAndEnd();
 					if (DPIndex < FriendlyNames.Num())
 					{
-						Item.OptionalFriendlyNameOverride = FText::FromString(FriendlyNames[DPIndex].TrimStartAndEnd());
+						FTextStringHelper::ReadFromBuffer(*FriendlyNames[DPIndex].TrimStartAndEnd(), Item.OptionalFriendlyNameOverride);
 					}
 					else if (DeviceProfileNames.Num() > 1)
 					{
