Move gc.MaxObjectsInGame, Substrate  and NumPreallocatedVertices settings from Android Platform ini to Base ini

diff --git a/Engine/Config/Android/AndroidEngine.ini b/Engine/Config/Android/AndroidEngine.ini
--- a/Engine/Config/Android/AndroidEngine.ini
+++ b/Engine/Config/Android/AndroidEngine.ini
@@ -9,9 +9,6 @@ DefaultFilmGrainTextureName=""
 PlatformRequiresDataCrypto=True
 PakSigningRequired=False
 
-[/Script/Engine.GarbageCollectionSettings]
-gc.MaxObjectsInGame=131072
-
 [Audio]
 PlatformFormat=OGG
 PlatformStreamingFormat=OGG
@@ -38,9 +35,6 @@ PreviewDeviceProfileSelectionModule="AndroidDeviceProfileSelector"
 [Advertising]
 DefaultProviderName=AndroidAdvertising
 
-[SlateRenderer]
-NumPreallocatedVertices=200
-
 [/Script/TcpMessaging.TcpMessagingSettings]
 EnableTransport=True
 ListenEndpoint=127.0.0.1:6666
@@ -93,13 +87,6 @@ r.Vulkan.MemoryBacktrace=0
 [/Script/Engine.RendererSettings]
 r.DistanceFields=1
 r.TemporalAA.Mobile.UseCompute=0
-r.Substrate.RoughDiffuse=0
-r.Substrate.ShadingQuality=2
-r.Substrate.TileCoord8bits=1
-r.Substrate.SheenQuality=2
-r.Substrate.Glints=0
-r.Substrate.SpecularProfile=0
-r.Substrate.ClosuresPerPixel=1
 ; Disables parallel GDME by default due to not enough threads
 r.Visibility.DynamicMeshElements.Parallel=0
 
diff --git a/Engine/Config/Android/BaseAndroidEngine.ini b/Engine/Config/Android/BaseAndroidEngine.ini
--- a/Engine/Config/Android/BaseAndroidEngine.ini
+++ b/Engine/Config/Android/BaseAndroidEngine.ini
@@ -17,4 +17,19 @@ DefaultExternalCredentialTypeStr=Google
 ;PlatformModularFeatureName=AndroidFetchBackgroundDownload
 
 [BuildPatchTool]
-VerificationThreadCount=3
\ No newline at end of file
+VerificationThreadCount=3
+
+[/Script/Engine.GarbageCollectionSettings]
+gc.MaxObjectsInGame=131072
+
+[SlateRenderer]
+NumPreallocatedVertices=200
+
+[/Script/Engine.RendererSettings]
+r.Substrate.RoughDiffuse=0
+r.Substrate.ShadingQuality=2
+r.Substrate.TileCoord8bits=1
+r.Substrate.SheenQuality=2
+r.Substrate.Glints=0
+r.Substrate.SpecularProfile=0
+r.Substrate.ClosuresPerPixel=1
\ No newline at end of file
