- Remove vestigial config

#rnx

diff --git a/Engine/Config/Android/BaseAndroidEngine.ini b/Engine/Config/Android/BaseAndroidEngine.ini
--- a/Engine/Config/Android/BaseAndroidEngine.ini
+++ b/Engine/Config/Android/BaseAndroidEngine.ini
@@ -9,9 +9,6 @@ bSupportsInAppPurchasing=true
 [Audio]
 AudioCaptureModuleName=AudioCaptureAndroid
 
-[OnlineServices.EOS.Auth]
-DefaultExternalCredentialTypeStr=Google
-
 ;Uncomment to turn on AndroidBackgroundHttp through the AndroidFetchBackgroundDownload and AndroidBackgroundService modules
 ;[BackgroundHttp]
 ;PlatformModularFeatureName=AndroidFetchBackgroundDownload
diff --git a/Engine/Config/IOS/BaseIOSEngine.ini b/Engine/Config/IOS/BaseIOSEngine.ini
--- a/Engine/Config/IOS/BaseIOSEngine.ini
+++ b/Engine/Config/IOS/BaseIOSEngine.ini
@@ -82,9 +82,6 @@ DefaultPlatformService=IOSLocalNotification
 Enable=true
 Alignment=16384
 
-[OnlineServices.EOS.Auth]
-DefaultExternalCredentialTypeStr=Apple
-
 [OnlineSubsystem]
 DefaultPlatformService=IOS
 NativePlatformService=IOS
diff --git a/Engine/Config/Linux/BaseLinuxEngine.ini b/Engine/Config/Linux/BaseLinuxEngine.ini
--- a/Engine/Config/Linux/BaseLinuxEngine.ini
+++ b/Engine/Config/Linux/BaseLinuxEngine.ini
@@ -4,9 +4,6 @@ LocalPlatformName=LNX
 [PlatformPaths]
 UnrealBuildTool=Engine/Build/BatchFiles/Linux/Build.sh
 
-[OnlineServices.EOS.Auth]
-DefaultExternalCredentialTypeStr=Steam
-
 [CrashReportClient]
 bAgreeToCrashUpload=false
 
diff --git a/Engine/Config/Mac/BaseMacEngine.ini b/Engine/Config/Mac/BaseMacEngine.ini
--- a/Engine/Config/Mac/BaseMacEngine.ini
+++ b/Engine/Config/Mac/BaseMacEngine.ini
@@ -11,9 +11,6 @@ PlatformHeadroomDB=-6
 [DevOptions.Shaders]
 MaxShaderJobBatchSize=25
 
-[OnlineServices.EOS.Auth]
-DefaultExternalCredentialTypeStr=Steam
-
 [OnlineSubsystem]
 LocalPlatformName=MAC
 
diff --git a/Engine/Config/Windows/BaseWindowsEngine.ini b/Engine/Config/Windows/BaseWindowsEngine.ini
--- a/Engine/Config/Windows/BaseWindowsEngine.ini
+++ b/Engine/Config/Windows/BaseWindowsEngine.ini
@@ -7,9 +7,6 @@ LocalPlatformName=WIN
 [DeviceProfileManager]
 DeviceProfileSelectionModule="WindowsDeviceProfileSelector"
 
-[OnlineServices.EOS.Auth]
-DefaultExternalCredentialTypeStr=Steam
-
 [PCD3D_SM5]
 BindlessResources=Disabled
 BindlessSamplers=Disabled
