Changed the default value of InputDeviceMappingPolicy to PrimaryUserSharesKeyboardAndFirstGamepad(1) for Android.

diff --git a/Engine/Config/Android/BaseAndroidInput.ini b/Engine/Config/Android/BaseAndroidInput.ini
--- a/Engine/Config/Android/BaseAndroidInput.ini
+++ b/Engine/Config/Android/BaseAndroidInput.ini
@@ -1,3 +1,3 @@
 [InputPlatformSettings_Android InputPlatformSettings]
 MaxPlatformUserCount=4
-input.DeviceMappingPolicy=0
\ No newline at end of file
+input.DeviceMappingPolicy=1
\ No newline at end of file
