Disable Vulkan Submission AllowTimelineSemaphores as it seems to cause an FD leak on adreno drivers.
(by allan.bentham)

#rnx

diff --git a/Engine/Config/Android/AndroidEngine.ini b/Engine/Config/Android/AndroidEngine.ini
--- a/Engine/Config/Android/AndroidEngine.ini
+++ b/Engine/Config/Android/AndroidEngine.ini
@@ -83,6 +83,9 @@ r.HairStrands.Binding=0
 r.TemporalAA.R11G11B10History=0
 ; Do not backtrace vulkan memory allocations
 r.Vulkan.MemoryBacktrace=0
+; Disable Vulkan Submission AllowTimelineSemaphores as it seems to cause an FD leak on adreno drivers.
+r.Vulkan.Submission.AllowTimelineSemaphores=0
+
 
 [/Script/Engine.RendererSettings]
 r.DistanceFields=1
