Allow NDK 29

#rnx

diff --git a/Engine/Config/Android/Android_SDK.json b/Engine/Config/Android/Android_SDK.json
--- a/Engine/Config/Android/Android_SDK.json
+++ b/Engine/Config/Android/Android_SDK.json
@@ -1,7 +1,7 @@
 {
 	"MainVersion": "r25b",
 	"MinVersion": "r25b",
-	"MaxVersion": "r28",
+	"MaxVersion": "r29",
 
 	"AutoSDKDirectory": "-27",
 
