VulkanRHI: Add default values for RHI threading settings in Android

#rnx

diff --git a/Engine/Config/Android/AndroidEngine.ini b/Engine/Config/Android/AndroidEngine.ini
--- a/Engine/Config/Android/AndroidEngine.ini
+++ b/Engine/Config/Android/AndroidEngine.ini
@@ -85,7 +85,10 @@ r.TemporalAA.R11G11B10History=0
 r.Vulkan.MemoryBacktrace=0
 ; Disable Vulkan Submission AllowTimelineSemaphores as it seems to cause an FD leak on adreno drivers.
 r.Vulkan.Submission.AllowTimelineSemaphores=0
-
+; Disable RHI multi-threading until its effects are measured on older devices
+r.Vulkan.Submission.UseInterruptThread=0
+r.Vulkan.Submission.UseSubmissionThread=0
+r.Vulkan.RHIThread=1
 
 [/Script/Engine.RendererSettings]
 r.DistanceFields=1
