Rename "Android_XR" flavor to "Android_OpenXR" to avoid confusion with the operating system.

diff --git a/Engine/Config/Android/DataDrivenPlatformInfo.ini b/Engine/Config/Android/DataDrivenPlatformInfo.ini
--- a/Engine/Config/Android/DataDrivenPlatformInfo.ini
+++ b/Engine/Config/Android/DataDrivenPlatformInfo.ini
@@ -42,7 +42,7 @@ ActiveIconName=LevelEditor.PreviewMode.AndroidVulkan.Enabled
 InactiveIconName=LevelEditor.PreviewMode.AndroidVulkan.Disabled
 MenuTooltip=NSLOCTEXT("PreviewPlatform", "PreviewMenuTooltip_AndroidVulkan", "Mobile preview using Android Vulkan quality settings.")
 IconText=NSLOCTEXT("PreviewPlatform", "PreviewIconText_AndroidVulkan", "Android Vulkan")
-DeviceProfileName=Android_Preview_Vulkan:Android_Low:Android_Mid:Android_High:Android_XR
+DeviceProfileName=Android_Preview_Vulkan:Android_Low:Android_Mid:Android_High:Android_OpenXR
 FriendlyName=NSLOCTEXT("PreviewPlatform", "PreviewFriendlyName_GenericVulkanAndroid", "Generic Android Vulkan"):NSLOCTEXT("PreviewPlatform", "PreviewFriendlyName_AndroidVulkanLow","Android Vulkan Low"):NSLOCTEXT("PreviewPlatform", "PreviewFriendlyName_AndroidVulkanMid","Android Vulkan Mid"):NSLOCTEXT("PreviewPlatform", "PreviewFriendlyName_AndroidVulkanHigh","Android Vulkan High"):NSLOCTEXT("PreviewPlatform", "PreviewFriendlyName_AndroidGenericXR","Android Generic XR")
 
 
diff --git a/Engine/Config/BaseDeviceProfiles.ini b/Engine/Config/BaseDeviceProfiles.ini
--- a/Engine/Config/BaseDeviceProfiles.ini
+++ b/Engine/Config/BaseDeviceProfiles.ini
@@ -158,7 +158,7 @@
 +DeviceProfileNameAndTypes=Meta_Quest_Pro,Android
 +DeviceProfileNameAndTypes=Oculus_Quest2,Android
 +DeviceProfileNameAndTypes=Oculus_Quest,Android
-+DeviceProfileNameAndTypes=Android_XR,Android
++DeviceProfileNameAndTypes=Android_OpenXR,Android
 +DeviceProfileNameAndTypes=MagicLeap_Vulkan,Android
 
 [GlobalDefaults DeviceProfile]
@@ -1250,7 +1250,7 @@ BaseProfileName=Android_High
 ; enable Vulkan
 +CVars=r.Android.DisableVulkanSupport=0
 
-[Android_XR DeviceProfile]
+[Android_OpenXR DeviceProfile]
 DeviceType=Android
 BaseProfileName=Android_Mid
 +CVars=vr.MobileMultiView=1
@@ -1275,7 +1275,7 @@ BaseProfileName=Oculus_Quest
 
 [Oculus_Quest DeviceProfile]
 DeviceType=Android
-BaseProfileName=Android_XR
+BaseProfileName=Android_OpenXR
 +CVars=xr.VRS.DynamicFoveation=1
 +CVars=r.Android.DisableVulkanSupport=0
 ; disable running Niagara compute shaders on the GPU. This can cause hard to trace crashes.
diff --git a/Engine/Source/Developer/Android/AndroidTargetPlatformControls/Private/AndroidTargetPlatformControls.h b/Engine/Source/Developer/Android/AndroidTargetPlatformControls/Private/AndroidTargetPlatformControls.h
--- a/Engine/Source/Developer/Android/AndroidTargetPlatformControls/Private/AndroidTargetPlatformControls.h
+++ b/Engine/Source/Developer/Android/AndroidTargetPlatformControls/Private/AndroidTargetPlatformControls.h
@@ -248,17 +248,17 @@ public:
 };
 
 // XR using the ASTC TPC now
-class FAndroid_XRTargetPlatformControls : public FAndroid_ASTCTargetPlatformControls
+class FAndroid_OpenXRTargetPlatformControls : public FAndroid_ASTCTargetPlatformControls
 {
 public:
-	FAndroid_XRTargetPlatformControls(bool bIsClient, ITargetPlatformSettings* TargetPlatformSettings)
-		: FAndroid_ASTCTargetPlatformControls(bIsClient, TargetPlatformSettings, TEXT("XR"))
+	FAndroid_OpenXRTargetPlatformControls(bool bIsClient, ITargetPlatformSettings* TargetPlatformSettings)
+		: FAndroid_ASTCTargetPlatformControls(bIsClient, TargetPlatformSettings, TEXT("OpenXR"))
 	{
 	}
 
 	virtual FString CookingDeviceProfileName() const override
 	{
-		return TEXT("Android_XR");
+		return TEXT("Android_OpenXR");
 	}
 };
 
diff --git a/Engine/Source/Developer/Android/AndroidTargetPlatformControls/Private/AndroidTargetPlatformControlsModule.cpp b/Engine/Source/Developer/Android/AndroidTargetPlatformControls/Private/AndroidTargetPlatformControlsModule.cpp
--- a/Engine/Source/Developer/Android/AndroidTargetPlatformControls/Private/AndroidTargetPlatformControlsModule.cpp
+++ b/Engine/Source/Developer/Android/AndroidTargetPlatformControls/Private/AndroidTargetPlatformControlsModule.cpp
@@ -58,7 +58,7 @@ protected:
 				SinglePlatforms.Add(new FAndroid_ETC2TargetPlatformControls(bIsClient, OutMap[EAndroidTextureFormatCategory::ETC2]));
 				
 				// XR target platform
-				SinglePlatforms.Add(new FAndroid_XRTargetPlatformControls(bIsClient, OutXR));
+				SinglePlatforms.Add(new FAndroid_OpenXRTargetPlatformControls(bIsClient, OutXR));
 
 				// these are used in NotifyMultiSelectedFormatsChanged, so track in another array
 				MultiPlatforms.Add(new FAndroid_MultiTargetPlatformControls(bIsClient, OutMulti));
diff --git a/Engine/Source/Developer/Android/AndroidTargetPlatformSettings/Private/AndroidTargetPlatformSettingsModule.cpp b/Engine/Source/Developer/Android/AndroidTargetPlatformSettings/Private/AndroidTargetPlatformSettingsModule.cpp
--- a/Engine/Source/Developer/Android/AndroidTargetPlatformSettings/Private/AndroidTargetPlatformSettingsModule.cpp
+++ b/Engine/Source/Developer/Android/AndroidTargetPlatformSettings/Private/AndroidTargetPlatformSettingsModule.cpp
@@ -66,9 +66,9 @@ protected:
 		TargetPlatforms.Add(MultiTP);
 		MultiTargetPlatformSettings = MultiTP;
 
-		ITargetPlatformSettings* XRTP = new FAndroid_XRTargetPlatformSettings();
-		TargetPlatforms.Add(XRTP);
-		XRTargetPlatformSettings = XRTP;
+		ITargetPlatformSettings* OpenXRTP = new FAndroid_OpenXRTargetPlatformSettings();
+		TargetPlatforms.Add(OpenXRTP);
+		XRTargetPlatformSettings = OpenXRTP;
 	}
 private:
 	ITargetPlatformSettings* MultiTargetPlatformSettings;
diff --git a/Engine/Source/Developer/Android/AndroidTargetPlatformSettings/Public/AndroidTargetPlatformSettings.h b/Engine/Source/Developer/Android/AndroidTargetPlatformSettings/Public/AndroidTargetPlatformSettings.h
--- a/Engine/Source/Developer/Android/AndroidTargetPlatformSettings/Public/AndroidTargetPlatformSettings.h
+++ b/Engine/Source/Developer/Android/AndroidTargetPlatformSettings/Public/AndroidTargetPlatformSettings.h
@@ -216,11 +216,11 @@ public:
 	}
 };
 
-class FAndroid_XRTargetPlatformSettings : public FAndroid_ASTCTargetPlatformSettings
+class FAndroid_OpenXRTargetPlatformSettings : public FAndroid_ASTCTargetPlatformSettings
 {
 
 public:
-	FAndroid_XRTargetPlatformSettings() : FAndroid_ASTCTargetPlatformSettings(TEXT("XR"))
+	FAndroid_OpenXRTargetPlatformSettings() : FAndroid_ASTCTargetPlatformSettings(TEXT("OpenXR"))
 	{
 	}
 };
diff --git a/Engine/Source/Programs/AutomationTool/Android/AndroidPlatform.Automation.cs b/Engine/Source/Programs/AutomationTool/Android/AndroidPlatform.Automation.cs
--- a/Engine/Source/Programs/AutomationTool/Android/AndroidPlatform.Automation.cs
+++ b/Engine/Source/Programs/AutomationTool/Android/AndroidPlatform.Automation.cs
@@ -4427,14 +4427,14 @@ public class AndroidPlatformASTC : AndroidPlatform
     }
 }
 
-public class AndroidPlatformXR : AndroidPlatform
+public class AndroidPlatformOpenXR : AndroidPlatform
 {
 	public override string GetCookPlatform(bool bDedicatedServer, bool bIsClientOnly)
 	{
-		return bIsClientOnly ? "Android_XRClient" : "Android_XR";
+		return bIsClientOnly ? "Android_OpenXRClient" : "Android_OpenXR";
 	}
 	public override TargetPlatformDescriptor GetTargetPlatformDescriptor()
 	{
-		return new TargetPlatformDescriptor(TargetPlatformType, "XR");
+		return new TargetPlatformDescriptor(TargetPlatformType, "OpenXR");
 	}
 }
diff --git a/Engine/Source/Programs/UnrealBuildTool/Platform/Android/UEDeployAndroid.cs b/Engine/Source/Programs/UnrealBuildTool/Platform/Android/UEDeployAndroid.cs
--- a/Engine/Source/Programs/UnrealBuildTool/Platform/Android/UEDeployAndroid.cs
+++ b/Engine/Source/Programs/UnrealBuildTool/Platform/Android/UEDeployAndroid.cs
@@ -485,7 +485,7 @@ namespace UnrealBuildTool
 
 			bool result = Ini.GetBool("/Script/AndroidRuntimeSettings.AndroidRuntimeSettings", "bPackageForMetaQuest", out bool bPackageForMetaQuest);
 
-			return (result && bPackageForMetaQuest) || bTargetOculusDevices || CookFlavor == "_XR";
+			return (result && bPackageForMetaQuest) || bTargetOculusDevices || CookFlavor == "_OpenXR";
 		}
 
 		public bool DisableVerifyOBBOnStartUp(ConfigHierarchy? Ini = null)
@@ -2803,7 +2803,7 @@ namespace UnrealBuildTool
 						bDXTEnabled = true;
 						break;
 					case "_ASTC":
-					case "_XR":
+					case "_OpenXR":
 						bASTCEnabled = true;
 						break;
 					default:
