[Deformer Graph & ComputeFramework] enable compute framework on more platforms

diff --git a/Engine/Config/Android/DataDrivenPlatformInfo.ini b/Engine/Config/Android/DataDrivenPlatformInfo.ini
--- a/Engine/Config/Android/DataDrivenPlatformInfo.ini
+++ b/Engine/Config/Android/DataDrivenPlatformInfo.ini
@@ -76,6 +76,7 @@ bSupportsUniformBufferObjects = true
 bSupportsDxc=true
 bSupportsIndependentSamplers=false
 bIsHlslcc=false
+bSupportsComputeFramework = true
 FriendlyName=LOCTEXT("FriendlyShaderPlatformName_Android", "Android OpenGL Mobile")
 
 [ShaderPlatform VULKAN_ES3_1_ANDROID]
@@ -135,5 +136,6 @@ bSupportsWaveOperations=RuntimeDependent
 MinimumWaveSize=4
 MaximumWaveSize=128
 bSupportsWave64=true
-FriendlyName=LOCTEXT("FriendlyShaderPlatformName_AndroidVulkanSM5", "Android Vulkan SM5 Mobile")
 bSupportsComputeFramework = true
+FriendlyName=LOCTEXT("FriendlyShaderPlatformName_AndroidVulkanSM5", "Android Vulkan SM5 Mobile")
+
diff --git a/Engine/Config/Mac/DataDrivenPlatformInfo.ini b/Engine/Config/Mac/DataDrivenPlatformInfo.ini
--- a/Engine/Config/Mac/DataDrivenPlatformInfo.ini
+++ b/Engine/Config/Mac/DataDrivenPlatformInfo.ini
@@ -119,6 +119,7 @@ bSupportsVolumeTextureAtomics=false
 bSupportsShaderPipelines = false
 bSupportsDualSourceBlending = true
 bSupportsIndependentSamplers=true
+bSupportsComputeFramework = true
 FriendlyName=LOCTEXT("FriendlyShaderPlatformName_MacMobile", "Mac Metal Mobile")
 
 [PreviewPlatform METAL_SM5]
diff --git a/Engine/Config/VulkanPC/DataDrivenPlatformInfo.ini b/Engine/Config/VulkanPC/DataDrivenPlatformInfo.ini
--- a/Engine/Config/VulkanPC/DataDrivenPlatformInfo.ini
+++ b/Engine/Config/VulkanPC/DataDrivenPlatformInfo.ini
@@ -277,7 +277,7 @@ bSupportsVertexShaderSRVs = false
 bSupportsUniformBufferObjects = true
 bSupportsDualSourceBlending = true
 bSupportsIndependentSamplers=true
-
+bSupportsComputeFramework = true
 FriendlyName=LOCTEXT("FriendlyShaderPlatformName_VulkanPC_Mobile", "PC Vulkan Mobile")
 
 [PreviewPlatform VULKAN_SM5]
diff --git a/Engine/Config/Windows/DataDrivenPlatformInfo.ini b/Engine/Config/Windows/DataDrivenPlatformInfo.ini
--- a/Engine/Config/Windows/DataDrivenPlatformInfo.ini
+++ b/Engine/Config/Windows/DataDrivenPlatformInfo.ini
@@ -194,4 +194,5 @@ bSupportsShaderPipelines = false
 bSupportsVertexShaderSRVs = false
 bSupportsUniformBufferObjects = true
 bSupportsIndependentSamplers=false
+bSupportsComputeFramework=true
 FriendlyName=LOCTEXT("FriendlyShaderPlatformName_PCGL_Mobile", "PC OpenGL Mobile")
