Bump Android main NDK to r27c

diff --git a/Engine/Config/Android/Android_SDK.json b/Engine/Config/Android/Android_SDK.json
--- a/Engine/Config/Android/Android_SDK.json
+++ b/Engine/Config/Android/Android_SDK.json
@@ -1,12 +1,12 @@
 {
-	"MainVersion": "r25b",
+	"MainVersion": "r27c",
 	"MinVersion": "r25b",
 	"MaxVersion": "r29",
 
-	"AutoSDKDirectory": "-27",
+	"AutoSDKDirectory": "-28",
 
 	"platforms": "android-34",
 	"build-tools": "34.0.0",
 	"cmake": "3.22.1",
-	"ndk": "25.1.8937393"
+	"ndk": "27.2.12479018"
 }
