Update Android AutoSDK  MinVersion to r27c

#android

diff --git a/Engine/Config/Android/Android_SDK.json b/Engine/Config/Android/Android_SDK.json
--- a/Engine/Config/Android/Android_SDK.json
+++ b/Engine/Config/Android/Android_SDK.json
@@ -1,6 +1,6 @@
 {
 	"MainVersion": "r27c",
-	"MinVersion": "r25b",
+	"MinVersion": "r27c",
 	"MaxVersion": "r29",
 
 	"AutoSDKDirectory": "-28",
