Lower Android SDK "MaxVersion" to r28c due to r29 compiler bug

#android

diff --git a/Engine/Config/Android/Android_SDK.json b/Engine/Config/Android/Android_SDK.json
--- a/Engine/Config/Android/Android_SDK.json
+++ b/Engine/Config/Android/Android_SDK.json
@@ -1,7 +1,7 @@
 {
 	"MainVersion": "r27c",
 	"MinVersion": "r27c",
-	"MaxVersion": "r29",
+	"MaxVersion": "r28c",
 
 	"AutoSDKDirectory": "-28",
 
