Enable distance fields by default on mobile platforms, except OpenGL

diff --git a/Engine/Config/Android/DataDrivenPlatformInfo.ini b/Engine/Config/Android/DataDrivenPlatformInfo.ini
--- a/Engine/Config/Android/DataDrivenPlatformInfo.ini
+++ b/Engine/Config/Android/DataDrivenPlatformInfo.ini
@@ -77,6 +77,7 @@ bSupportsDxc=true
 bSupportsIndependentSamplers=false
 bIsHlslcc=false
 bSupportsComputeFramework = true
+bSupportsDistanceFields=false
 FriendlyName=LOCTEXT("FriendlyShaderPlatformName_Android", "Android OpenGL Mobile")
 
 [ShaderPlatform VULKAN_ES3_1_ANDROID]
@@ -101,6 +102,7 @@ bSupportsDxc=true
 bSupportsIndependentSamplers=true
 bIsHlslcc=false
 bSupportsComputeFramework = true
+bSupportsDistanceFields=true
 
 FriendlyName=LOCTEXT("FriendlyShaderPlatformName_AndroidVulkan", "Android Vulkan Mobile")
 
diff --git a/Engine/Config/IOS/BaseIOSEngine.ini b/Engine/Config/IOS/BaseIOSEngine.ini
--- a/Engine/Config/IOS/BaseIOSEngine.ini
+++ b/Engine/Config/IOS/BaseIOSEngine.ini
@@ -9,7 +9,7 @@ DefaultFilmGrainTextureName=""
 gc.MaxObjectsInGame=131072
 
 [/Script/Engine.RendererSettings]
-r.DistanceFields=0
+r.DistanceFields=1
 r.TemporalAA.Mobile.UseCompute=0
 ; Disables parallel GDME by default due to not enough threads
 r.Visibility.DynamicMeshElements.Parallel=0
diff --git a/Engine/Config/IOS/DataDrivenPlatformInfo.ini b/Engine/Config/IOS/DataDrivenPlatformInfo.ini
--- a/Engine/Config/IOS/DataDrivenPlatformInfo.ini
+++ b/Engine/Config/IOS/DataDrivenPlatformInfo.ini
@@ -58,6 +58,7 @@ bSupportsVolumeTextureAtomics=false
 bSupportsShaderPipelines = false
 bSupportsIndependentSamplers=true
 bSupportsComputeFramework = true
+bSupportsDistanceFields=true
 FriendlyName=LOCTEXT("FriendlyShaderPlatformName_IOSMobile", "IOS Metal Mobile")
 
 [ShaderPlatform METAL_SM5_IOS]
diff --git a/Engine/Config/Windows/DataDrivenPlatformInfo.ini b/Engine/Config/Windows/DataDrivenPlatformInfo.ini
--- a/Engine/Config/Windows/DataDrivenPlatformInfo.ini
+++ b/Engine/Config/Windows/DataDrivenPlatformInfo.ini
@@ -148,7 +148,7 @@ MaxFeatureLevel=ES3_1
 ShaderFormat=PCD3D_ES31
 bIsPC=true
 bIsMobile=true
-bSupportsDistanceFields = false
+bSupportsDistanceFields=true
 bSupportsComputeFramework=true
 bSupportsInstancedStereo=true
 bSupportsVertexShaderLayer=true
