Mobile: Enable AO for High bucket. Enable dynamic inset shadows for High and Epic

diff --git a/Engine/Config/Android/AndroidScalability.ini b/Engine/Config/Android/AndroidScalability.ini
--- a/Engine/Config/Android/AndroidScalability.ini
+++ b/Engine/Config/Android/AndroidScalability.ini
@@ -74,6 +74,7 @@ r.Shadow.RadiusThreshold=0.06
 r.VolumetricFog=0
 r.ContactShadows=0
 r.DFShadowQuality=0
+r.Mobile.SupportInsetShadows=True
 
 [ShadowQuality@3]
 r.ShadowQuality=2
@@ -84,6 +85,7 @@ r.Shadow.RadiusThreshold=0.06
 r.VolumetricFog=0
 r.ContactShadows=0
 r.DFShadowQuality=2
+r.Mobile.SupportInsetShadows=True
 
 ;-----------------------------------------------------------------------------------------------------------------
 
@@ -130,7 +132,7 @@ r.BloomQuality=2
 r.DepthOfFieldQuality=1
 r.LightShaftQuality=1
 r.EyeAdaptationQuality=1
-r.Mobile.AmbientOcclusionQuality=0
+r.Mobile.AmbientOcclusionQuality=2
 r.LensFlareQuality=2
 
 [PostProcessQuality@3]
@@ -139,7 +141,7 @@ r.BloomQuality=3
 r.DepthOfFieldQuality=1
 r.LightShaftQuality=1
 r.EyeAdaptationQuality=1
-r.Mobile.AmbientOcclusionQuality=2
+r.Mobile.AmbientOcclusionQuality=3
 r.LensFlareQuality=2
 
 [AntiAliasingQuality@0]
diff --git a/Engine/Config/IOS/IOSScalability.ini b/Engine/Config/IOS/IOSScalability.ini
--- a/Engine/Config/IOS/IOSScalability.ini
+++ b/Engine/Config/IOS/IOSScalability.ini
@@ -74,6 +74,7 @@ r.Shadow.RadiusThreshold=0.06
 r.VolumetricFog=0
 r.ContactShadows=0
 r.DFShadowQuality=0
+r.Mobile.SupportInsetShadows=True
 
 [ShadowQuality@3]
 r.ShadowQuality=2
@@ -84,6 +85,7 @@ r.Shadow.RadiusThreshold=0.06
 r.VolumetricFog=0
 r.ContactShadows=0
 r.DFShadowQuality=2
+r.Mobile.SupportInsetShadows=True
 
 ;-----------------------------------------------------------------------------------------------------------------
 
@@ -130,7 +132,7 @@ r.BloomQuality=2
 r.DepthOfFieldQuality=0
 r.LightShaftQuality=0
 r.EyeAdaptationQuality=1
-r.Mobile.AmbientOcclusionQuality=0
+r.Mobile.AmbientOcclusionQuality=2
 r.LensFlareQuality=2
 
 [PostProcessQuality@3]
@@ -139,7 +141,7 @@ r.BloomQuality=3
 r.DepthOfFieldQuality=1
 r.LightShaftQuality=1
 r.EyeAdaptationQuality=1
-r.Mobile.AmbientOcclusionQuality=2
+r.Mobile.AmbientOcclusionQuality=3
 r.LensFlareQuality=2
 
 [AntiAliasingQuality@0]
diff --git a/Engine/Source/Runtime/Renderer/Private/ShadowSetup.cpp b/Engine/Source/Runtime/Renderer/Private/ShadowSetup.cpp
--- a/Engine/Source/Runtime/Renderer/Private/ShadowSetup.cpp
+++ b/Engine/Source/Runtime/Renderer/Private/ShadowSetup.cpp
@@ -388,7 +388,7 @@ static TAutoConsoleVariable<bool> CVarMobileInsetShadows(
 	TEXT("r.Mobile.SupportInsetShadows"),
 	false,
 	TEXT("Enables inset shadows on mobile. Currently not supported on platforms not using ShadowMaskTexture. Default = false."),
-	ECVF_RenderThreadSafe
+	ECVF_Scalability | ECVF_RenderThreadSafe
 );
 
 /** Whether to round the shadow map up to power of two on mobile platform. */
