VulkanRHI: use renames instead of copy+barrier for ub updates on Android.

#rnx

diff --git a/Engine/Config/Android/AndroidEngine.ini b/Engine/Config/Android/AndroidEngine.ini
--- a/Engine/Config/Android/AndroidEngine.ini
+++ b/Engine/Config/Android/AndroidEngine.ini
@@ -97,6 +97,8 @@ r.Vulkan.DiagnosticBuffer=1
 r.Vulkan.EnableTransientResourceAllocator=0
 ; Disable sparse binding and residency
 r.Vulkan.AllowSparseResidency=0
+; Use renames instead of copy+barrier on mobile devices
+r.Vulkan.AllowUniformUpload=0
 
 ; enable stripping debug strings from streamable manager in shipping to save memory
 s.StreamableStripDebugName=true
