Mobile: Use ?Low? quality TAA by default on Android and iOS
Use the pixel shader TAA path on platforms where it is beneficial

#Synced-CL 50642587

diff --git a/Engine/Config/Android/AndroidScalability.ini b/Engine/Config/Android/AndroidScalability.ini
--- a/Engine/Config/Android/AndroidScalability.ini
+++ b/Engine/Config/Android/AndroidScalability.ini
@@ -148,18 +148,22 @@ r.LensFlareQuality=2
 
 [AntiAliasingQuality@0]
 r.Mobile.AntiAliasing=0
+r.TemporalAA.Quality=0
 r.FXAA.Quality=0
 
 [AntiAliasingQuality@1]
 r.Mobile.AntiAliasing=1
+r.TemporalAA.Quality=0
 r.FXAA.Quality=0
 
 [AntiAliasingQuality@2]
 r.Mobile.AntiAliasing=2
+r.TemporalAA.Quality=0
 r.FXAA.Quality=0
 
 [AntiAliasingQuality@3]
 r.Mobile.AntiAliasing=2
+r.TemporalAA.Quality=0
 r.FXAA.Quality=0
 
 ; ------------------------------------------------------------------------------
diff --git a/Engine/Config/IOS/IOSScalability.ini b/Engine/Config/IOS/IOSScalability.ini
--- a/Engine/Config/IOS/IOSScalability.ini
+++ b/Engine/Config/IOS/IOSScalability.ini
@@ -148,18 +148,22 @@ r.LensFlareQuality=2
 
 [AntiAliasingQuality@0]
 r.Mobile.AntiAliasing=0
+r.TemporalAA.Quality=0
 r.FXAA.Quality=0
 
 [AntiAliasingQuality@1]
 r.Mobile.AntiAliasing=1
+r.TemporalAA.Quality=0
 r.FXAA.Quality=0
 
 [AntiAliasingQuality@2]
 r.Mobile.AntiAliasing=2
+r.TemporalAA.Quality=0
 r.FXAA.Quality=0
 
 [AntiAliasingQuality@3]
 r.Mobile.AntiAliasing=2
+r.TemporalAA.Quality=0
 r.FXAA.Quality=0
 
 ; ------------------------------------------------------------------------------
