Enable TAA upscaling on Android and iOS to reduce blurriness

diff --git a/Engine/Config/Android/AndroidEngine.ini b/Engine/Config/Android/AndroidEngine.ini
--- a/Engine/Config/Android/AndroidEngine.ini
+++ b/Engine/Config/Android/AndroidEngine.ini
@@ -76,8 +76,6 @@ r.Forward.MaxCulledLightsPerCell = 8
 r.HairStrands.Simulation=0
 r.HairStrands.Strands=0
 r.HairStrands.Binding=0
-; R11G11B10 UAV is not supported on Android
-r.TemporalAA.R11G11B10History=0
 ; Do not backtrace vulkan memory allocations
 r.Vulkan.MemoryBacktrace=0
 ; Disable Vulkan Submission AllowTimelineSemaphores as it seems to cause an FD leak on adreno drivers.
